- Fixed ModuleManager crash in UnrealFileServer

diff --git a/Engine/Source/Runtime/NetworkFileSystem/Private/NetworkFileServer.h b/Engine/Source/Runtime/NetworkFileSystem/Private/NetworkFileServer.h
--- a/Engine/Source/Runtime/NetworkFileSystem/Private/NetworkFileServer.h
+++ b/Engine/Source/Runtime/NetworkFileSystem/Private/NetworkFileServer.h
@@ -22,7 +22,8 @@ public:
 	 * @param InPort - The port number to bind to (0 = any available port).
 	 * @param InFileRequestDelegate -
 	 */
-	FNetworkFileServer( int32 InPort, const FFileRequestDelegate* InFileRequestDelegate, const FRecompileShadersDelegate* InRecompileShadersDelegate );
+	FNetworkFileServer( int32 InPort, const FFileRequestDelegate* InFileRequestDelegate, 
+		const FRecompileShadersDelegate* InRecompileShadersDelegate, const TArray<ITargetPlatform*>& InActiveTargetPlatforms );
 
 	/**
 	 * Destructor.
@@ -96,6 +97,9 @@ public:
 	// Holds a delegate to be invoked when a client requests a shader recompile.
 	FRecompileShadersDelegate RecompileShadersDelegate;
 
+	// cached copy of the active target platforms (if any)
+	const TArray<ITargetPlatform*> ActiveTargetPlatforms;
+
 	// Holds the address that the server is bound to.
 	TSharedPtr<FInternetAddr> ListenAddr;
 };
\ No newline at end of file
diff --git a/Engine/Source/Runtime/NetworkFileSystem/Private/NetworkFileServerConnection.cpp b/Engine/Source/Runtime/NetworkFileSystem/Private/NetworkFileServerConnection.cpp
--- a/Engine/Source/Runtime/NetworkFileSystem/Private/NetworkFileServerConnection.cpp
+++ b/Engine/Source/Runtime/NetworkFileSystem/Private/NetworkFileServerConnection.cpp
@@ -12,12 +12,14 @@
 /* FNetworkFileServerClientConnection structors
  *****************************************************************************/
 
-FNetworkFileServerClientConnection::FNetworkFileServerClientConnection( FSocket* InSocket, const FFileRequestDelegate& InFileRequestDelegate, const FRecompileShadersDelegate& InRecompileShadersDelegate )
+FNetworkFileServerClientConnection::FNetworkFileServerClientConnection( FSocket* InSocket, const FFileRequestDelegate& InFileRequestDelegate, 
+		const FRecompileShadersDelegate& InRecompileShadersDelegate, const TArray<ITargetPlatform*>& InActiveTargetPlatforms )
 	: LastHandleId(0)
 	, MCSocket(NULL)
 	, Sandbox(NULL)
 	, Socket(InSocket)
 	, bNeedsToStop(false)
+	, ActiveTargetPlatforms(InActiveTargetPlatforms)
 {
 	if (InFileRequestDelegate.IsBound())
 	{
@@ -696,15 +698,14 @@ void FNetworkFileServerClientConnection::ProcessGetFileList( FArchive& In, FArch
 
 	ConnectedPlatformName = TEXT("");
 	// figure out the best matching target platform for the set of valid ones
-	ITargetPlatformManagerModule& TPM = GetTargetPlatformManagerRef();
-	static const TArray<ITargetPlatform*>& ActiveTargetPlatforms =  TPM.GetTargetPlatforms();
-	for (int32 TPIndex = 0; TPIndex < TargetPlatformNames.Num(); TPIndex++)
+	for (int32 TPIndex = 0; TPIndex < TargetPlatformNames.Num() && ConnectedPlatformName == TEXT(""); TPIndex++)
 	{
 		UE_LOG(LogFileServer, Display, TEXT("    Possible Target Platform from client: %s"), *TargetPlatformNames[TPIndex]);
 
 		// look for a matching target platform
 		for (int32 ActiveTPIndex = 0; ActiveTPIndex < ActiveTargetPlatforms.Num(); ActiveTPIndex++)
 		{
+			UE_LOG(LogFileServer, Display, TEXT("   Checking against: %s"), *ActiveTargetPlatforms[ActiveTPIndex]->PlatformName());
 			if (ActiveTargetPlatforms[ActiveTPIndex]->PlatformName() == TargetPlatformNames[TPIndex])
 			{
 				ConnectedPlatformName = ActiveTargetPlatforms[ActiveTPIndex]->PlatformName();
diff --git a/Engine/Source/Runtime/NetworkFileSystem/Private/NetworkFileServerConnection.h b/Engine/Source/Runtime/NetworkFileSystem/Private/NetworkFileServerConnection.h
--- a/Engine/Source/Runtime/NetworkFileSystem/Private/NetworkFileServerConnection.h
+++ b/Engine/Source/Runtime/NetworkFileSystem/Private/NetworkFileServerConnection.h
@@ -22,7 +22,8 @@ public:
 	 * @param InSocket - The client socket to use.
 	 * @param InFileRequestDelegate - A delegate to be invoked when the client requests a file.
 	 */
-	FNetworkFileServerClientConnection( FSocket* InSocket, const FFileRequestDelegate& InFileRequestDelegate, const FRecompileShadersDelegate& InRecompileShadersDelegate );
+	FNetworkFileServerClientConnection( FSocket* InSocket, const FFileRequestDelegate& InFileRequestDelegate, 
+		const FRecompileShadersDelegate& InRecompileShadersDelegate, const TArray<ITargetPlatform*>& InActiveTargetPlatforms );
 
 	/**
 	 * Destructor.
@@ -221,6 +222,9 @@ private:
 	// Holds a delegate to be invoked when a client requests a shader recompile.
 	FRecompileShadersDelegate RecompileShadersDelegate;
 
+	// cached copy of the active target platforms (if any)
+	const TArray<ITargetPlatform*>& ActiveTargetPlatforms;
+
 	/** Holds a flag indicating that the thread should be stopped. */
 	bool bNeedsToStop;
 };
\ No newline at end of file
diff --git a/Engine/Source/Runtime/NetworkFileSystem/Private/NetworkFileSystemModule.cpp b/Engine/Source/Runtime/NetworkFileSystem/Private/NetworkFileSystemModule.cpp
--- a/Engine/Source/Runtime/NetworkFileSystem/Private/NetworkFileSystemModule.cpp
+++ b/Engine/Source/Runtime/NetworkFileSystem/Private/NetworkFileSystemModule.cpp
@@ -28,7 +28,18 @@ public:
 			Port = DEFAULT_FILE_SERVING_PORT;
 		}
 
-		return new FNetworkFileServer(Port, InFileRequestDelegate, InRecompileShadersDelegate);
+		TArray<ITargetPlatform*> ActiveTargetPlatforms;
+
+		// only bother getting the target platforms if there was "-targetplatform" on the commandline, otherwise UnrealFileServer will 
+		// log out some scary sounding, but innocuous logs
+		FString Platforms;
+		if (FParse::Value(FCommandLine::Get(), TEXT("TARGETPLATFORM="), Platforms))
+		{
+			ITargetPlatformManagerModule& TPM = GetTargetPlatformManagerRef();
+			ActiveTargetPlatforms =  TPM.GetActiveTargetPlatforms();
+		}
+
+		return new FNetworkFileServer(Port, InFileRequestDelegate, InRecompileShadersDelegate, ActiveTargetPlatforms);
 	}
 
 	// End INetworkFileSystemModule interface
