- Missed this file in last checkin

diff --git a/Engine/Source/Runtime/NetworkFileSystem/Private/NetworkFileServer.cpp b/Engine/Source/Runtime/NetworkFileSystem/Private/NetworkFileServer.cpp
--- a/Engine/Source/Runtime/NetworkFileSystem/Private/NetworkFileServer.cpp
+++ b/Engine/Source/Runtime/NetworkFileSystem/Private/NetworkFileServer.cpp
@@ -5,13 +5,16 @@
 =============================================================================*/
 
 #include "NetworkFileSystemPrivatePCH.h"
+#include "TargetPlatform.h"
 
 
 /* FNetworkFileServer constructors
  *****************************************************************************/
 
-FNetworkFileServer::FNetworkFileServer( int32 InPort, const FFileRequestDelegate* InFileRequestDelegate, const FRecompileShadersDelegate* InRecompileShadersDelegate )
+FNetworkFileServer::FNetworkFileServer( int32 InPort, const FFileRequestDelegate* InFileRequestDelegate, 
+			const FRecompileShadersDelegate* InRecompileShadersDelegate, const TArray<ITargetPlatform*>& InActiveTargetPlatforms )
 	: bNeedsToStop(false)
+	, ActiveTargetPlatforms(InActiveTargetPlatforms)
 {
 	UE_LOG(LogFileServer, Warning, TEXT("Unreal Network File Server starting up..."));
 
@@ -123,7 +126,7 @@ uint32 FNetworkFileServer::Run( )
 
 			if (ClientSocket != NULL)
 			{
-				FNetworkFileServerClientConnection* Connection = new FNetworkFileServerClientConnection(ClientSocket, FileRequestDelegate, RecompileShadersDelegate);
+				FNetworkFileServerClientConnection* Connection = new FNetworkFileServerClientConnection(ClientSocket, FileRequestDelegate, RecompileShadersDelegate, ActiveTargetPlatforms);
 				Connections.Add(Connection);
 				UE_LOG(LogFileServer, Display, TEXT( "Client %s connected." ), *Connection->GetDescription() );
 			}
