Fixes for QAGame non unity compile

diff --git a/Engine/Source/Editor/Levels/Private/LevelViewModel.cpp b/Engine/Source/Editor/Levels/Private/LevelViewModel.cpp
--- a/Engine/Source/Editor/Levels/Private/LevelViewModel.cpp
+++ b/Engine/Source/Editor/Levels/Private/LevelViewModel.cpp
@@ -2,6 +2,7 @@
 
 #include "LevelsPrivatePCH.h"
 
+#include "LevelViewModel.h"
 #include "LevelUtils.h"
 #include "EditorLevelUtils.h"
 #include "Toolkits/AssetEditorManager.h"
diff --git a/Engine/Source/Runtime/NetworkFileSystem/Private/NetworkFileSystemModule.cpp b/Engine/Source/Runtime/NetworkFileSystem/Private/NetworkFileSystemModule.cpp
--- a/Engine/Source/Runtime/NetworkFileSystem/Private/NetworkFileSystemModule.cpp
+++ b/Engine/Source/Runtime/NetworkFileSystem/Private/NetworkFileSystemModule.cpp
@@ -6,6 +6,9 @@
 
 #include "NetworkFileSystemPrivatePCH.h"
 
+#include "TargetDeviceId.h"
+#include "ITargetDevice.h"
+#include "ITargetPlatformManagerModule.h"
 #include "ModuleManager.h"
 
 DEFINE_LOG_CATEGORY(LogFileServer);
