Don't send DumpShaderDebugInfo to devices when serving files

diff --git a/Engine/Source/Runtime/NetworkFileSystem/Private/NetworkFileServerConnection.cpp b/Engine/Source/Runtime/NetworkFileSystem/Private/NetworkFileServerConnection.cpp
--- a/Engine/Source/Runtime/NetworkFileSystem/Private/NetworkFileServerConnection.cpp
+++ b/Engine/Source/Runtime/NetworkFileSystem/Private/NetworkFileServerConnection.cpp
@@ -804,6 +804,7 @@ void FNetworkFileServerClientConnection::ProcessGetFileList( FArchive& In, FArch
 		DirectoriesToSkip.Add(FString(RootDirectories[DirIndex] / TEXT("Saved/Config")));
 		DirectoriesToSkip.Add(FString(RootDirectories[DirIndex] / TEXT("Saved/Logs")));
 		DirectoriesToSkip.Add(FString(RootDirectories[DirIndex] / TEXT("Saved/Sandboxes")));
+		DirectoriesToSkip.Add(FString(RootDirectories[DirIndex] / TEXT("Saved/ShaderDebugInfo")));
 		DirectoriesToSkip.Add(FString(RootDirectories[DirIndex] / TEXT("Intermediate")));
 		DirectoriesToSkip.Add(FString(RootDirectories[DirIndex] / TEXT("Documentation")));
 		DirectoriesToSkip.Add(FString(RootDirectories[DirIndex] / TEXT("Extras")));
