Shortening Cook sandbox path from Game\\Saved\\Sandboxes\\Cooked-Platform to Game\\Saved\\Cooked\\Platform

diff --git a/Engine/Build/BatchFiles/HTML5/HTMLSetupCookedContent-Swinger.bat b/Engine/Build/BatchFiles/HTML5/HTMLSetupCookedContent-Swinger.bat
--- a/Engine/Build/BatchFiles/HTML5/HTMLSetupCookedContent-Swinger.bat
+++ b/Engine/Build/BatchFiles/HTML5/HTMLSetupCookedContent-Swinger.bat
@@ -1,5 +1,5 @@
 set GAME=Samples\SampleGames\SwingNinja
-set DEST=%GAME%\Saved\Sandboxes\Cooked-HTML5
+set DEST=%GAME%\Saved\Cooked\HTML5
 
 pushd ..\..\..\..\
 xcopy /y /i /d Engine\Config\*.ini %DEST%\Engine\Config
@@ -13,6 +13,6 @@ xcopy /y /i /d %GAME%\Config\*.ini %DEST%\%GAME%\Config
 xcopy /y /i /d %GAME%\Config\HTML5\*.ini %DEST%\%GAME%\Config\HTML5
 xcopy /y /i /d %GAME%\Content\Localization\*.* %DEST%\%GAME%\Content\Localization\*.*
 
-%EMSCRIPTEN%\tools\file_packager.py Engine\Binaries\HTML5\SwingNinja.data --preload  %GAME%\Config Engine\Config  %GAME%\Saved\Sandboxes\Cooked-HTML5@\ --pre-run --js-output=Engine\Binaries\HTML5\SwingNinja.data.js
+%EMSCRIPTEN%\tools\file_packager.py Engine\Binaries\HTML5\SwingNinja.data --preload  %GAME%\Config Engine\Config  %GAME%\Saved\Cooked\HTML5@\ --pre-run --js-output=Engine\Binaries\HTML5\SwingNinja.data.js
 
 popd
diff --git a/Engine/Build/BatchFiles/HTML5/HTMLSetupCookedContent.bat b/Engine/Build/BatchFiles/HTML5/HTMLSetupCookedContent.bat
--- a/Engine/Build/BatchFiles/HTML5/HTMLSetupCookedContent.bat
+++ b/Engine/Build/BatchFiles/HTML5/HTMLSetupCookedContent.bat
@@ -1,5 +1,5 @@
 set GAME=QAGame
-set DEST=%GAME%\Saved\Sandboxes\Cooked-HTML5
+set DEST=%GAME%\Saved\Cooked\HTML5
 
 pushd ..\..\..\..\
 xcopy /y /i /d Engine\Config\*.ini %DEST%\Engine\Config
@@ -13,6 +13,6 @@ xcopy /y /i /d %GAME%\Config\*.ini %DEST%\%GAME%\Config
 xcopy /y /i /d %GAME%\Config\HTML5\*.ini %DEST%\%GAME%\Config\HTML5
 xcopy /y /i /d %GAME%\Content\Localization\*.* %DEST%\%GAME%\Content\Localization\*.*
 
-%EMSCRIPTEN%\tools\file_packager.py QAGame\Binaries\HTML5\QAGame.data --preload  QAGame\Config Engine\Config  QAGame\Saved\Sandboxes\Cooked-HTML5@\ --pre-run --js-output=QAGame\Binaries\HTML5\QAGame.data.js
+%EMSCRIPTEN%\tools\file_packager.py QAGame\Binaries\HTML5\QAGame.data --preload  QAGame\Config Engine\Config  QAGame\Saved\Cooked\HTML5@\ --pre-run --js-output=QAGame\Binaries\HTML5\QAGame.data.js
 
 popd
diff --git a/Engine/Config/Mac/MacEngine.ini b/Engine/Config/Mac/MacEngine.ini
--- a/Engine/Config/Mac/MacEngine.ini
+++ b/Engine/Config/Mac/MacEngine.ini
@@ -1,9 +1,6 @@
 [Audio]
 AudioDeviceModuleName=CoreAudio
 
-[Core.System]
-SeekFreePCPaths=../../UDKGame/CookedMac
-
 [DerivedDataBackendGraph]
 Shared=(Type=FileSystem, ReadOnly=false, Clean=false, Flush=false, DeleteUnused=true, UnusedFileAge=38, FoldersToClean=-1, Path=/Volumes/UE4DDC, EnvPathOverride=UE-SharedDataCachePath)
 AltShared=(Type=FileSystem, ReadOnly=true, Clean=false, Flush=false, DeleteUnused=true, UnusedFileAge=38, FoldersToClean=-1, Path=/Volumes/UE4DDC2, EnvPathOverride=UE-SharedDataCachePath2)
diff --git a/Engine/Source/Editor/UnrealEd/Private/Commandlets/CookCommandlet.cpp b/Engine/Source/Editor/UnrealEd/Private/Commandlets/CookCommandlet.cpp
--- a/Engine/Source/Editor/UnrealEd/Private/Commandlets/CookCommandlet.cpp
+++ b/Engine/Source/Editor/UnrealEd/Private/Commandlets/CookCommandlet.cpp
@@ -571,12 +571,14 @@ FString UCookCommandlet::GetOutputDirectoryOverride() const
 	// Output directory override.	
 	if (!FParse::Value(*Params, TEXT("Output="), OutputDirectory))
 	{
-		OutputDirectory = TEXT("Cooked-[Platform]");
+		// Full path so that the sandbox wrapper doesn't try to re-base it under Sandboxes
+		OutputDirectory = FPaths::Combine(*FPaths::GameDir(), TEXT("Saved"), TEXT("Cooked"), TEXT("[Platform]"));
+		OutputDirectory = FPaths::ConvertRelativePathToFull(OutputDirectory);
 	}
 	else if (!OutputDirectory.Contains(TEXT("[Platform]"), ESearchCase::IgnoreCase, ESearchDir::FromEnd) )
 	{
 		// Output directory needs to contain [Platform] token to be able to cook for multiple targets.
-		OutputDirectory += TEXT("/Cooked-[Platform]");
+		OutputDirectory = FPaths::Combine(*OutputDirectory, TEXT("[Platform]"));
 	}
 	FPaths::NormalizeDirectoryName(OutputDirectory);
 
diff --git a/Engine/Source/Editor/UnrealEd/Private/CookOnTheFlyServer.cpp b/Engine/Source/Editor/UnrealEd/Private/CookOnTheFlyServer.cpp
--- a/Engine/Source/Editor/UnrealEd/Private/CookOnTheFlyServer.cpp
+++ b/Engine/Source/Editor/UnrealEd/Private/CookOnTheFlyServer.cpp
@@ -727,14 +727,16 @@ FString UCookOnTheFlyServer::GetOutputDirectoryOverride( const FString &OutputDi
 {
 	FString OutputDirectory = OutputDirectoryOverride;
 	// Output directory override.	
-	if ( OutputDirectory.Len() <= 0 )
+	if (OutputDirectory.Len() <= 0)
 	{
-		OutputDirectory = TEXT("Cooked-[Platform]");
+		// Full path so that the sandbox wrapper doesn't try to re-base it under Sandboxes
+		OutputDirectory = FPaths::Combine(*FPaths::GameDir(), TEXT("Saved"), TEXT("Cooked"), TEXT("[Platform]"));
+		OutputDirectory = FPaths::ConvertRelativePathToFull(OutputDirectory);
 	}
 	else if (!OutputDirectory.Contains(TEXT("[Platform]"), ESearchCase::IgnoreCase, ESearchDir::FromEnd) )
 	{
 		// Output directory needs to contain [Platform] token to be able to cook for multiple targets.
-		OutputDirectory += TEXT("/Cooked-[Platform]");
+		OutputDirectory = FPaths::Combine(*OutputDirectory, TEXT("[Platform]"));
 	}
 	FPaths::NormalizeDirectoryName(OutputDirectory);
 
diff --git a/Engine/Source/Programs/AutomationTool/Scripts/CookCommand.Automation.cs b/Engine/Source/Programs/AutomationTool/Scripts/CookCommand.Automation.cs
--- a/Engine/Source/Programs/AutomationTool/Scripts/CookCommand.Automation.cs
+++ b/Engine/Source/Programs/AutomationTool/Scripts/CookCommand.Automation.cs
@@ -101,11 +101,11 @@ public partial class Project : CommandUtils
 	private static void CleanupCookedData(List<string> PlatformsToCook, ProjectParams Params)
 	{		
 		var ProjectPath = Path.GetFullPath(Params.RawProjectPath);
-		var CookedSandboxesPath = CombinePaths(GetDirectoryName(ProjectPath), "Saved", "Sandboxes", "Cooked-");
+		var CookedSandboxesPath = CombinePaths(GetDirectoryName(ProjectPath), "Saved", "Cooked");
 		var CleanDirs = new string[PlatformsToCook.Count];
 		for (int DirIndex = 0; DirIndex < CleanDirs.Length; ++DirIndex)
 		{
-			CleanDirs[DirIndex] = CookedSandboxesPath + PlatformsToCook[DirIndex];
+			CleanDirs[DirIndex] = CombinePaths(CookedSandboxesPath, PlatformsToCook[DirIndex]);
 		}
 
 		const bool bQuiet = true;
diff --git a/Engine/Source/Programs/AutomationTool/Scripts/CopyBuildToStagingDirectory.Automation.cs b/Engine/Source/Programs/AutomationTool/Scripts/CopyBuildToStagingDirectory.Automation.cs
--- a/Engine/Source/Programs/AutomationTool/Scripts/CopyBuildToStagingDirectory.Automation.cs
+++ b/Engine/Source/Programs/AutomationTool/Scripts/CopyBuildToStagingDirectory.Automation.cs
@@ -200,7 +200,7 @@ public partial class Project : CommandUtils
 			SC.StageFiles(StagedFileTypeForMovies, CombinePaths(SC.ProjectRoot, "Content/Movies"), "*", true, null, CombinePaths(SC.RelativeProjectRootForStage, "Content/Movies"), true, !Params.Pak);
 
 			// eliminate the sand box
-			SC.StageFiles(StagedFileType.UFS, CombinePaths(SC.ProjectRoot, "Saved", "Sandboxes", "Cooked-" + SC.CookPlatform), "*", true, null, "", true, !Params.Pak);
+			SC.StageFiles(StagedFileType.UFS, CombinePaths(SC.ProjectRoot, "Saved", "Cooked", SC.CookPlatform), "*", true, null, "", true, !Params.Pak);
 
 			// CrashReportClient is a standalone slate app that does not look in the generated pak file, so it needs the Content/Slate and Shaders/StandaloneRenderer folders Non-UFS
 			// @todo Make CrashReportClient more portable so we don't have to do this
diff --git a/Engine/Source/Programs/AutomationTool/Scripts/GUBP.Automation.cs b/Engine/Source/Programs/AutomationTool/Scripts/GUBP.Automation.cs
--- a/Engine/Source/Programs/AutomationTool/Scripts/GUBP.Automation.cs
+++ b/Engine/Source/Programs/AutomationTool/Scripts/GUBP.Automation.cs
@@ -1977,7 +1977,7 @@ public class GUBP : BuildCommand
 
         public override string RootIfAnyForTempStorage()
         {
-            return CombinePaths(Path.GetDirectoryName(GameProj.FilePath), "Saved", "Sandboxes", "Cooked-" + CookPlatform);
+            return CombinePaths(Path.GetDirectoryName(GameProj.FilePath), "Saved", "Cooked", CookPlatform);
         }
         public override void DoBuild(GUBP bp)
         {
diff --git a/Engine/Source/Programs/AutomationTool/Scripts/RunProjectCommand.Automation.cs b/Engine/Source/Programs/AutomationTool/Scripts/RunProjectCommand.Automation.cs
--- a/Engine/Source/Programs/AutomationTool/Scripts/RunProjectCommand.Automation.cs
+++ b/Engine/Source/Programs/AutomationTool/Scripts/RunProjectCommand.Automation.cs
@@ -635,7 +635,7 @@ public partial class Project : CommandUtils
 			}
 			else if (!Params.Stage)
 			{
-				var SandboxPath = CombinePaths(SC.RuntimeProjectRootDir, "Saved/Sandboxes", "Cooked-" + SC.CookPlatform);
+				var SandboxPath = CombinePaths(SC.RuntimeProjectRootDir, "Saved", "Cooked", SC.CookPlatform);
 				if (!SC.StageTargetPlatform.LaunchViaUFE)
 				{
 					TempCmdLine += "-sandbox=" + CommandUtils.MakePathSafeToUseWithCommandLine(SandboxPath) + " ";
diff --git a/Engine/Source/Programs/IOS/iPhonePackager/Config.cs b/Engine/Source/Programs/IOS/iPhonePackager/Config.cs
--- a/Engine/Source/Programs/IOS/iPhonePackager/Config.cs
+++ b/Engine/Source/Programs/IOS/iPhonePackager/Config.cs
@@ -77,7 +77,7 @@ namespace iPhonePackager
 		/// </summary>
 		public static string CookedDirectory
 		{
-			get { return Path.GetFullPath(GameDirectory + @"\Saved\Sandboxes\Cooked-IOS"); }
+			get { return Path.GetFullPath(GameDirectory + @"\Saved\Cooked\IOS"); }
 		}
 
 		/// <summary>
diff --git a/Engine/Source/Runtime/Engine/Classes/AI/Navigation/CrowdFollowingComponent.h b/Engine/Source/Runtime/Engine/Classes/AI/Navigation/CrowdFollowingComponent.h
--- a/Engine/Source/Runtime/Engine/Classes/AI/Navigation/CrowdFollowingComponent.h
+++ b/Engine/Source/Runtime/Engine/Classes/AI/Navigation/CrowdFollowingComponent.h
@@ -1,4 +1,5 @@
 // Copyright 1998-2014 Epic Games, Inc. All Rights Reserved.
+#pragma once
 
 #include "CrowdFollowingComponent.generated.h"
 
diff --git a/Engine/Source/Runtime/Engine/Classes/AI/Navigation/CrowdManager.h b/Engine/Source/Runtime/Engine/Classes/AI/Navigation/CrowdManager.h
--- a/Engine/Source/Runtime/Engine/Classes/AI/Navigation/CrowdManager.h
+++ b/Engine/Source/Runtime/Engine/Classes/AI/Navigation/CrowdManager.h
@@ -1,5 +1,7 @@
 // Copyright 1998-2014 Epic Games, Inc. All Rights Reserved.
 
+#pragma once
+
 #include "CrowdManager.generated.h"
 
 /**
diff --git a/Engine/Source/Runtime/NetworkFileSystem/Private/NetworkFileServerConnection.cpp b/Engine/Source/Runtime/NetworkFileSystem/Private/NetworkFileServerConnection.cpp
--- a/Engine/Source/Runtime/NetworkFileSystem/Private/NetworkFileServerConnection.cpp
+++ b/Engine/Source/Runtime/NetworkFileSystem/Private/NetworkFileServerConnection.cpp
@@ -750,7 +750,7 @@ void FNetworkFileServerClientConnection::ProcessGetFileList( FArchive& In, FArch
 	if ( FPaths::IsProjectFilePathSet() )
 	{
 		FString ProjectDir = FPaths::GetPath(FPaths::GetProjectFilePath());
-		SandboxDirectory = FPaths::Combine(*ProjectDir, TEXT("Saved"), TEXT("Sandboxes"), *(FString("Cooked-") + ConnectedPlatformName));
+		SandboxDirectory = FPaths::Combine(*ProjectDir, TEXT("Saved"), TEXT("Cooked"), *ConnectedPlatformName);
 		if( bIsStreamingRequest )
 		{
 			RootDirectories.Add(ProjectDir);
@@ -760,15 +760,16 @@ void FNetworkFileServerClientConnection::ProcessGetFileList( FArchive& In, FArch
 	{
 		if (FPaths::GetExtension(GameName) == IProjectManager::GetProjectFileExtension())
 		{
-			SandboxDirectory = FString::Printf(TEXT("%s/Saved/Sandboxes/Cooked-%s"), *FPaths::GetPath(GameName), *ConnectedPlatformName);
+			SandboxDirectory = FPaths::Combine(*FPaths::GetPath(GameName), TEXT("Saved"), TEXT("Cooked"), *ConnectedPlatformName);
 		}
 		else
 		{
 			//@todo: This assumes the game is located in the UE4 Root directory
-			FString SandboxFolder = FString::Printf(TEXT("Cooked-%s"), *ConnectedPlatformName);
-			SandboxDirectory = FPaths::Combine(*FPaths::GetRelativePathToRoot(), *GameName, TEXT("Saved/Sandboxes"), *SandboxFolder);
+			SandboxDirectory = FPaths::Combine(*FPaths::GetRelativePathToRoot(), *GameName, TEXT("Saved"), TEXT("Cooked"), *ConnectedPlatformName);
 		}
 	}
+	// Convert to full path so that the sandbox wrapper doesn't re-base to Saved/Sandboxes
+	SandboxDirectory = FPaths::ConvertRelativePathToFull(SandboxDirectory);
 
 	// delete any existing one first, in case game name somehow changed and client is re-asking for files (highly unlikely)
 	delete Sandbox;
diff --git a/Engine/Source/Runtime/SandboxFile/Private/IPlatformFileSandboxWrapper.cpp b/Engine/Source/Runtime/SandboxFile/Private/IPlatformFileSandboxWrapper.cpp
--- a/Engine/Source/Runtime/SandboxFile/Private/IPlatformFileSandboxWrapper.cpp
+++ b/Engine/Source/Runtime/SandboxFile/Private/IPlatformFileSandboxWrapper.cpp
@@ -13,6 +13,11 @@ FSandboxPlatformFile::FSandboxPlatformFile(bool bInEntireEngineWillUseThisSandbo
 {
 }
 
+static FString GetCookedSandboxDir()
+{
+	return FPaths::Combine(*(FPaths::GameDir()), TEXT("Cooked"), ANSI_TO_TCHAR(FPlatformProperties::PlatformName()));
+}
+
 bool FSandboxPlatformFile::ShouldBeUsed(IPlatformFile* Inner, const TCHAR* CmdLine) const
 {
 	FString SandboxDir;
@@ -20,8 +25,7 @@ bool FSandboxPlatformFile::ShouldBeUsed(IPlatformFile* Inner, const TCHAR* CmdLi
 #if PLATFORM_DESKTOP && (UE_GAME || UE_SERVER)
 	if (FPlatformProperties::RequiresCookedData() && SandboxDir.IsEmpty() && Inner == &FPlatformFileManager::Get().GetPlatformFile() && bEntireEngineWillUseThisSandbox)
 	{
-		FString SandboxName = FString(TEXT("Cooked-")) + ANSI_TO_TCHAR(FPlatformProperties::PlatformName());
-		SandboxDir = FPaths::Combine(*(FPaths::GameDir()), TEXT("Saved"), TEXT("Sandboxes"), *SandboxDir);
+		SandboxDir = GetCookedSandboxDir();
 		bResult = FPlatformFileManager::Get().GetPlatformFile().DirectoryExists(*SandboxDir);
 	}
 #endif
@@ -35,8 +39,7 @@ bool FSandboxPlatformFile::Initialize(IPlatformFile* Inner, const TCHAR* CmdLine
 #if PLATFORM_DESKTOP && (UE_GAME || UE_SERVER)
 	if (CommandLineDirectory.IsEmpty() && bEntireEngineWillUseThisSandbox)
 	{
-		FString SandboxName = FString(TEXT("Cooked-")) + ANSI_TO_TCHAR(FPlatformProperties::PlatformName());
-		CommandLineDirectory = FPaths::Combine(*(FPaths::GameDir()), TEXT("Saved"), TEXT("Sandboxes"), *SandboxName);
+		CommandLineDirectory = GetCookedSandboxDir();
 		UE_LOG(LogInit, Display, TEXT("No sandbox specified, assuming %s"), *CommandLineDirectory);
 	}
 #endif
