Fix up sandbox paths sent over the fileserver so they can be matched properly against locally cached files.  Integrating from parallel rendering branch.

diff --git a/Engine/Source/Runtime/NetworkFileSystem/Private/NetworkFileServerConnection.cpp b/Engine/Source/Runtime/NetworkFileSystem/Private/NetworkFileServerConnection.cpp
--- a/Engine/Source/Runtime/NetworkFileSystem/Private/NetworkFileServerConnection.cpp
+++ b/Engine/Source/Runtime/NetworkFileSystem/Private/NetworkFileServerConnection.cpp
@@ -145,6 +145,33 @@ void FNetworkFileServerClientConnection::ConvertClientFilenameToServerFilename(F
 	}
 }
 
+/**
+ * Fixup sandbox paths to match what package loading will request on the client side.  e.g.
+ * Sandbox path: "../../../Elemental/Content/Elemental/Effects/FX_Snow_Cracks/Crack_02/Materials/M_SnowBlast.uasset ->
+ * client path: "../../../Samples/Showcases/Elemental/Content/Elemental/Effects/FX_Snow_Cracks/Crack_02/Materials/M_SnowBlast.uasset"
+ * This ensures that devicelocal-cached files will be properly timestamp checked before deletion.
+ */
+static TMap<FString, FDateTime> FixupSandboxPathsForClient(FSandboxPlatformFile* Sandbox, const TMap<FString, FDateTime>& SandboxPaths, const FString& LocalEngineDir, const FString& LocalGameDir)
+{
+	TMap<FString, FDateTime> FixedFiletimes;
+	FString SandboxEngine = Sandbox->ConvertToSandboxPath(*LocalEngineDir) + TEXT("/");
+
+	// we need to add an extra bit to the game path to make the sandbox convert it correctly (investigate?)
+	// @todo: double check this
+	FString SandboxGame = Sandbox->ConvertToSandboxPath(*(LocalGameDir + TEXT("a.txt"))).Replace(TEXT("a.txt"), TEXT(""));
+
+	// since the sandbox remaps from A/B/C to C, and the client has no idea of this, we need to put the files
+	// into terms of the actual LocalGameDir, which is all that the client knows about
+	for (TMap<FString, FDateTime>::TConstIterator It(SandboxPaths); It; ++It)
+	{
+		FString Fixed = Sandbox->ConvertToSandboxPath(*It.Key());
+		Fixed = Fixed.Replace(*SandboxEngine, *LocalEngineDir);
+		Fixed = Fixed.Replace(*SandboxGame, *LocalGameDir);
+		FixedFiletimes.Add(Fixed, It.Value());
+	}
+	return FixedFiletimes;
+}
+
 void FNetworkFileServerClientConnection::ConvertServerFilenameToClientFilename(FString& FilenameToConvert)
 {
 	if (FilenameToConvert.StartsWith(FPaths::EngineDir()))
@@ -831,7 +858,8 @@ void FNetworkFileServerClientConnection::ProcessGetFileList( FArchive& In, FArch
 	Out << LocalGameDir;
 
 	// return the files and their timestamps
-	Out << Visitor.FileTimes;
+	TMap<FString, FDateTime> FixedTimes = FixupSandboxPathsForClient(Sandbox, Visitor.FileTimes, LocalEngineDir, LocalGameDir);
+	Out << FixedTimes;
 
 	// Do it again, preventing access to non-cooked files
 	if( bIsStreamingRequest == false )
@@ -857,7 +885,8 @@ void FNetworkFileServerClientConnection::ProcessGetFileList( FArchive& In, FArch
 		}
 	
 		// return the cached files and their timestamps
-		Out << VisitorForCacheDates.FileTimes;
+		FixedTimes = FixupSandboxPathsForClient(Sandbox, VisitorForCacheDates.FileTimes, LocalEngineDir, LocalGameDir);
+		Out << FixedTimes;
 	}
 }
 
