Fixing NetworkFileServer hang / stack overflow

diff --git a/Engine/Source/Runtime/NetworkFileSystem/Private/NetworkFileServerConnection.cpp b/Engine/Source/Runtime/NetworkFileSystem/Private/NetworkFileServerConnection.cpp
--- a/Engine/Source/Runtime/NetworkFileSystem/Private/NetworkFileServerConnection.cpp
+++ b/Engine/Source/Runtime/NetworkFileSystem/Private/NetworkFileServerConnection.cpp
@@ -832,6 +832,7 @@ void FNetworkFileServerClientConnection::ProcessGetFileList( FArchive& In, FArch
 		DirectoriesToSkip.Add(FString(RootDirectories[DirIndex] / TEXT("Saved/Config")));
 		DirectoriesToSkip.Add(FString(RootDirectories[DirIndex] / TEXT("Saved/Logs")));
 		DirectoriesToSkip.Add(FString(RootDirectories[DirIndex] / TEXT("Saved/Sandboxes")));
+		DirectoriesToSkip.Add(FString(RootDirectories[DirIndex] / TEXT("Saved/Cooked")));
 		DirectoriesToSkip.Add(FString(RootDirectories[DirIndex] / TEXT("Saved/ShaderDebugInfo")));
 		DirectoriesToSkip.Add(FString(RootDirectories[DirIndex] / TEXT("Intermediate")));
 		DirectoriesToSkip.Add(FString(RootDirectories[DirIndex] / TEXT("Documentation")));
