Fix assert to allow for heartbeat messages to preceed the first get file list message

diff --git a/Engine/Source/Runtime/NetworkFileSystem/Private/NetworkFileServerConnection.cpp b/Engine/Source/Runtime/NetworkFileSystem/Private/NetworkFileServerConnection.cpp
--- a/Engine/Source/Runtime/NetworkFileSystem/Private/NetworkFileServerConnection.cpp
+++ b/Engine/Source/Runtime/NetworkFileSystem/Private/NetworkFileServerConnection.cpp
@@ -127,7 +127,7 @@ void FNetworkFileServerClientConnection::ProcessPayload(FArchive& Ar,FArchive& O
 	NFS_Messages::Type Msg = NFS_Messages::Type(Cmd);
 
 	// make sure the first thing is GetFileList which initializes the game/platform
-	checkf(Msg == NFS_Messages::GetFileList || Sandbox != NULL, TEXT("The first client message MUST be GetFileList, not %d"), (int32)Msg);
+	checkf(Msg == NFS_Messages::GetFileList || Msg == NFS_Messages::Heartbeat || Sandbox != NULL, TEXT("The first client message MUST be GetFileList, not %d"), (int32)Msg);
 
 	// process the message!
 	bool bSendUnsolicitedFiles = false;
