Fix cookonthefly crashes from unsolicited files piling up.

diff --git a/Engine/Source/Runtime/NetworkFileSystem/Private/NetworkFileServerConnection.cpp b/Engine/Source/Runtime/NetworkFileSystem/Private/NetworkFileServerConnection.cpp
--- a/Engine/Source/Runtime/NetworkFileSystem/Private/NetworkFileServerConnection.cpp
+++ b/Engine/Source/Runtime/NetworkFileSystem/Private/NetworkFileServerConnection.cpp
@@ -236,13 +236,14 @@ void FNetworkFileServerClientConnection::ProcessPayload(FArchive& Ar,FArchive& O
 	if (Out.TotalSize() && bSendUnsolicitedFiles)
 	{
 		int32 NumUnsolictedFiles = UnsolictedFiles.Num();
-			Out << NumUnsolictedFiles;
+		Out << NumUnsolictedFiles;
 
-			for (int32 Index = 0; Index < NumUnsolictedFiles; Index++)
-			{
+		for (int32 Index = 0; Index < NumUnsolictedFiles; Index++)
+		{
 			PackageFile(UnsolictedFiles[Index], Out);
 		}
-			}
+		UnsolictedFiles.Empty();		
+	}
 
 	UE_LOG(LogFileServer, Verbose, TEXT("Done Processing payload with Cmd %d Total Size sending %d "), Cmd,Out.TotalSize());
 }
