#UE4 #HTML5

   #TTP 341694 disable http nfs by default for all editor platforms - enable manually when needed for time being. a recent refactor enabled it for some builds.

    -code cleanup/adhere to standards

diff --git a/Engine/Source/Programs/UnrealBuildTool/HTML5/UEBuildHTML5.cs b/Engine/Source/Programs/UnrealBuildTool/HTML5/UEBuildHTML5.cs
--- a/Engine/Source/Programs/UnrealBuildTool/HTML5/UEBuildHTML5.cs
+++ b/Engine/Source/Programs/UnrealBuildTool/HTML5/UEBuildHTML5.cs
@@ -18,10 +18,6 @@ namespace UnrealBuildTool
         [XmlConfig]
         public static bool EnableHTTPForNFS = false;
 
-        // Actually use HTTP code path for both Client and Server. 
-        [XmlConfig]
-        public static bool UseHTTPForNFS = false;
-
         static private string TargetPlatformName = "HTML5";
 
 		// This is the SDK version we support in //CarefullyRedist.
@@ -149,9 +145,6 @@ namespace UnrealBuildTool
                 }
             }
 
-            // if HTTP is not enabled, we can't use it. 
-            if (EnableHTTPForNFS == false)
-                UseHTTPForNFS = false; 
         }
 
         /**
@@ -351,14 +344,17 @@ namespace UnrealBuildTool
 
                 if (InModule.ToString() == "NetworkFile")
                 {
-                    InModule.AddPublicDefinition("ENABLE_HTTP_FOR_NFS=1");
-                    if (Target.Architecture == "-win32")
-                    {
-                        InModule.AddPrivateDependencyModule("HTML5Win32");
-                    }
-                    else
-                    {
-                        InModule.AddPrivateDependencyModule("HTML5JS");
+                    if (EnableHTTPForNFS)
+                    {   
+                        InModule.AddPublicDefinition("ENABLE_HTTP_FOR_NFS=1");
+                        if (Target.Architecture == "-win32")
+                        {
+                            InModule.AddPrivateDependencyModule("HTML5Win32");
+                        }
+                        else
+                        {
+                            InModule.AddPrivateDependencyModule("HTML5JS");
+                        }
                     }
                 }
             }
@@ -371,15 +367,18 @@ namespace UnrealBuildTool
                 {
                     InModule.AddPlatformSpecificDynamicallyLoadedModule("HTML5TargetPlatform");
                 }
-                if (InModule.ToString() == "NetworkFile") // client
+                if (EnableHTTPForNFS)
                 {
-                    InModule.AddPrivateDependencyModule("HTTP");
-                    InModule.AddPublicDefinition("ENABLE_HTTP_FOR_NFS=1");
-                }
-                else if (InModule.ToString() == "NetworkFileSystem") // server 
-                {
-                    InModule.AddPublicDependencyModule("WebSockets");
-                    InModule.AddPublicDefinition("ENABLE_HTTP_FOR_NFS=1");
+                    if (InModule.ToString() == "NetworkFile") // client
+                    {
+                        InModule.AddPrivateDependencyModule("HTTP");
+                        InModule.AddPublicDefinition("ENABLE_HTTP_FOR_NFS=1");
+                    }
+                    else if (InModule.ToString() == "NetworkFileSystem") // server 
+                    {
+                        InModule.AddPublicDependencyModule("WebSockets");
+                        InModule.AddPublicDefinition("ENABLE_HTTP_FOR_NFS=1");
+                    }
                 }
             }
         }
diff --git a/Engine/Source/Runtime/NetworkFileSystem/Private/NetworkFileServerHttp.cpp b/Engine/Source/Runtime/NetworkFileSystem/Private/NetworkFileServerHttp.cpp
--- a/Engine/Source/Runtime/NetworkFileSystem/Private/NetworkFileServerHttp.cpp
+++ b/Engine/Source/Runtime/NetworkFileSystem/Private/NetworkFileServerHttp.cpp
@@ -22,7 +22,7 @@ struct PerSessionData
 
 
 // protocol array. 
-static struct libwebsocket_protocols protocols[] = {
+static struct libwebsocket_protocols Protocols[] = {
 	/* first protocol must always be HTTP handler */
 	{
 			"http-only",   // name
@@ -93,28 +93,28 @@ bool FNetworkFileServerHttp::Init()
 	lws_set_log_level( LLL_ERR | LLL_WARN | LLL_NOTICE | LLL_DEBUG , libwebsocket_debugLog);
 #endif 
 
-	struct lws_context_creation_info info;
-	memset(&info,0,sizeof(lws_context_creation_info));
+	struct lws_context_creation_info Info;
+	memset(&Info,0,sizeof(lws_context_creation_info));
 
 	// look up libwebsockets.h for details. 
-	info.port = Port;
+	Info.port = Port;
 	// we listen on all available interfaces. 
-	info.iface = NULL;
+	Info.iface = NULL;
 	// serve only the http protocols. 
-	info.protocols = protocols;
+	Info.protocols = Protocols;
 	// no extensions
-	info.extensions = NULL; 
-	info.gid = -1;
-	info.uid = -1;
-	info.options = 0;
+	Info.extensions = NULL; 
+	Info.gid = -1;
+	Info.uid = -1;
+	Info.options = 0;
 	// tack on this object. 
-	info.user = this; 
+	Info.user = this; 
 	
-	context = libwebsocket_create_context(&info);
+	Context = libwebsocket_create_context(&Info);
 
-	Port = info.port;
+	Port = Info.port;
 
-	if (context == NULL) {
+	if (Context == NULL) {
 		UE_LOG(LogFileServer, Fatal, TEXT(" Could not create a libwebsocket content for port : %d"), Port);
 		return false;
 	}
@@ -134,7 +134,7 @@ FString FNetworkFileServerHttp::GetSupportedProtocol() const
 bool FNetworkFileServerHttp::GetAddressList(TArray<TSharedPtr<FInternetAddr> >& OutAddresses) const
 {
 	// if Init failed, its already too late. 
-	ensure( context != nullptr);
+	ensure( Context != nullptr);
 
 	// we are listening to all local interfaces. 
 	ISocketSubsystem::Get()->GetLocalAdapterAddresses(OutAddresses);
@@ -173,8 +173,8 @@ uint32 FNetworkFileServerHttp::Run()
 	while(!StopRequested.GetValue())
 	{
 		// service libwebsocket, have a slight delay so it doesn't spin on zero load. 
-		libwebsocket_service(context, 10); 
-		libwebsocket_callback_on_writable_all_protocol(&protocols[0]);
+		libwebsocket_service(Context, 10); 
+		libwebsocket_callback_on_writable_all_protocol(&Protocols[0]);
 	}
 
 	UE_LOG(LogFileServer, Display, TEXT("Unreal Network File Http Server is now Shutting down "));
@@ -191,8 +191,8 @@ void FNetworkFileServerHttp::Exit()
 {
 	// let's start shutting down. 
 	// fires a LWS_CALLBACK_PROTOCOL_DESTROY callback, we clean up after ourselves there.
-	libwebsocket_context_destroy(context);
-	context = NULL;
+	libwebsocket_context_destroy(Context);
+	Context = NULL;
 }
 
 FNetworkFileServerHttp::~FNetworkFileServerHttp()
@@ -204,7 +204,7 @@ FNetworkFileServerHttp::~FNetworkFileServerHttp()
 		delete Element.Value; 
 	}
 	// make sure context has been already cleaned up. 
-	check( context == NULL ); 
+	check( Context == NULL ); 
 }
 
 FNetworkFileServerClientConnection* FNetworkFileServerHttp::CreateNewConnection()
@@ -245,28 +245,28 @@ void FNetworkFileServerHttp::Process(FArchive& In, FArchive&Out, FNetworkFileSer
 // This static function handles all callbacks coming in and when context is services via libwebsocket_service
 // return value of -1, closes the connection.
 int FNetworkFileServerHttp::CallBack_HTTP(	
-			struct libwebsocket_context *context, 
-			struct libwebsocket *wsi, 
-			enum libwebsocket_callback_reasons reason, 
-			void *user, 
-			void *in, 
-			size_t len)
+			struct libwebsocket_context *Context, 
+			struct libwebsocket *Wsi, 
+			enum libwebsocket_callback_reasons Reason, 
+			void *User, 
+			void *In, 
+			size_t Len)
 {
-	PerSessionData* bufferInfo = (PerSessionData*)user;
-	FNetworkFileServerHttp* Server = (FNetworkFileServerHttp*)libwebsocket_context_user(context); 
+	PerSessionData* BufferInfo = (PerSessionData*)User;
+	FNetworkFileServerHttp* Server = (FNetworkFileServerHttp*)libwebsocket_context_user(Context); 
 
-	switch (reason)
+	switch (Reason)
 	{
 
 	case LWS_CALLBACK_HTTP: 
 
 		// hang on to socket even if there's no data for atleast 60 secs. 
-		libwebsocket_set_timeout(wsi, NO_PENDING_TIMEOUT, 60);
+		libwebsocket_set_timeout(Wsi, NO_PENDING_TIMEOUT, 60);
 
 		/* if it was not legal POST URL, let it continue and accept data */
-		if (!lws_hdr_total_length(wsi, WSI_TOKEN_POST_URI))
+		if (!lws_hdr_total_length(Wsi, WSI_TOKEN_POST_URI))
 		{
-			char *requested_uri = (char *) in;
+			char *requested_uri = (char *) In;
 
 			// client request the base page. e.g  http://unrealfileserver:port/ 
 			// just return a banner, probably add some more information, e,g Version, Config, Game. etc. 
@@ -286,14 +286,14 @@ int FNetworkFileServerHttp::CallBack_HTTP(
 					);
 
 				// very small data being sent, its fine to just send.  
-				libwebsocket_write(wsi,(unsigned char*)TCHAR_TO_ANSI(Buffer),FCStringAnsi::Strlen(TCHAR_TO_ANSI(Buffer)), LWS_WRITE_HTTP);
+				libwebsocket_write(Wsi,(unsigned char*)TCHAR_TO_ANSI(Buffer),FCStringAnsi::Strlen(TCHAR_TO_ANSI(Buffer)), LWS_WRITE_HTTP);
 			}
 			else
 			{
 				// client has asked for a file. ( only html/js files are served.) 
 					
 				// what type is being served. 
-				FString FilePath = FPaths::GameDir() / TEXT("Binaries/HTML5") +  FString((char*)in); 
+				FString FilePath = FPaths::GameDir() / TEXT("Binaries/HTML5") +  FString((ANSICHAR*)In); 
 				TCHAR *Mime = NULL; 
 
 				if ( FilePath.Contains(".html") )
@@ -319,9 +319,11 @@ int FNetworkFileServerHttp::CallBack_HTTP(
 				// rather than reading and sending in batches because Unreal NFS servers are not running in memory 
 				// constrained env and the added complexity is not worth it. 
 
-				FILE *f = fopen(TCHAR_TO_ANSI(*AbsoluteFilePath), "rb");
 
-				if (f == NULL )
+				TArray<uint8> FileData; 
+				FFileHelper::LoadFileToArray(FileData, *AbsoluteFilePath, FILEREAD_Silent);
+					
+				if (FileData.Num() == 0)
 				{
 					// umm. we didn't find file, we should tell the client that we couldn't find it. 
 					// send 404.
@@ -329,66 +331,58 @@ int FNetworkFileServerHttp::CallBack_HTTP(
 									"Server: Unreal File Server\x0d\x0a"
 									"Connection: close\x0d\x0a";
 
-					libwebsocket_write(wsi,(unsigned char*)Header,FCStringAnsi::Strlen(Header), LWS_WRITE_HTTP);
+					libwebsocket_write(Wsi,(unsigned char*)Header,FCStringAnsi::Strlen(Header), LWS_WRITE_HTTP);
 					// chug along, client will close the  connection. 
 					break; 
 				}
 
-				// find the size of the file. 
-				fseek(f, 0, SEEK_END);
-				long fsize = ftell(f);
-				fseek(f, 0, SEEK_SET);
-
 				// file up the header. 
-				TCHAR header[1024];
-				int x  = FCString::Sprintf(header,
+				TCHAR Header[1024];
+				int Length  = FCString::Sprintf(Header,
 					TEXT("HTTP/1.1 200 OK\x0d\x0a")
 					TEXT("Server: Unreal File Server\x0d\x0a")
 					TEXT("Connection: close\x0d\x0a")
 					TEXT("Content-Type: %s \x0d\x0a")
 					TEXT("Content-Length: %u\x0d\x0a\x0d\x0a"),
-					Mime,fsize);
+					Mime,FileData.Num());
 
 				// make space for the whole file in our out buffer. 
-				bufferInfo->Out.Append((const unsigned char*)TCHAR_TO_ANSI(header),x); 
-				bufferInfo->Out.AddUninitialized(fsize);
-				// read
-				fread(bufferInfo->Out.GetData() + x, fsize, 1, f);
-				fclose(f);
+				BufferInfo->Out.Append((uint8*)TCHAR_TO_ANSI(Header),Length); 
+				BufferInfo->Out.Append(FileData);
 				// we need to write back to the client, queue up a write callback. 
-				libwebsocket_callback_on_writable(context, wsi);
+				libwebsocket_callback_on_writable(Context, Wsi);
 			}
 		}
 		else
 		{
 			// we got a post request!,  queue up a write callback. 
-			libwebsocket_callback_on_writable(context, wsi);
+			libwebsocket_callback_on_writable(Context, Wsi);
 		}
 
 		break;
 	case LWS_CALLBACK_HTTP_BODY: 
 		{
 			// post data is coming in, push it on to our incoming buffer.  
-			UE_LOG(LogFileServer, Log, TEXT("Incoming HTTP Partial Body Size %d, total size  %d"),len, len+ bufferInfo->In.Num());
-			bufferInfo->In.Append((uint8*)in,len);
+			UE_LOG(LogFileServer, Log, TEXT("Incoming HTTP Partial Body Size %d, total size  %d"),Len, Len+ BufferInfo->In.Num());
+			BufferInfo->In.Append((uint8*)In,Len);
 			// we received some data - update time out.
-			libwebsocket_set_timeout(wsi, NO_PENDING_TIMEOUT, 60);
+			libwebsocket_set_timeout(Wsi, NO_PENDING_TIMEOUT, 60);
 		}
 		break;
 	case LWS_CALLBACK_HTTP_BODY_COMPLETION: 
 		{
 			// we have all the post data from the client. 
 			// create archives and process them. 
-			UE_LOG(LogFileServer, Log, TEXT("Incoming HTTP total size  %d"), bufferInfo->In.Num());
-			FMemoryReader Reader(bufferInfo->In);
+			UE_LOG(LogFileServer, Log, TEXT("Incoming HTTP total size  %d"), BufferInfo->In.Num());
+			FMemoryReader Reader(BufferInfo->In);
 			FBufferArchive Writer;
 
 			FNetworkFileServerHttp::Process(Reader,Writer,Server);
 
 			// even if we have 0 data to push, tell the client that we don't any data. 
-			char header[1024];
-			int x  = FCStringAnsi::Sprintf(
-				(ANSICHAR*)header,
+			ANSICHAR Header[1024];
+			int Length  = FCStringAnsi::Sprintf(
+				(ANSICHAR*)Header,
 				"HTTP/1.1 200 OK\x0d\x0a"
 				"Server: Unreal File Server\x0d\x0a"
 				"Connection: close\x0d\x0a"
@@ -398,21 +392,21 @@ int FNetworkFileServerHttp::CallBack_HTTP(
 				);
 
 			// Add Http Header 
-			bufferInfo->Out.Append((unsigned char*)header,x); 
+			BufferInfo->Out.Append((uint8*)Header,Length); 
 			// Add Binary Data.  
-			bufferInfo->Out.Append(Writer);
+			BufferInfo->Out.Append(Writer);
 
 			// we have enqueued data increase timeout and push a writable callback. 
-			libwebsocket_set_timeout(wsi, NO_PENDING_TIMEOUT, 60);
-			libwebsocket_callback_on_writable(context, wsi);
+			libwebsocket_set_timeout(Wsi, NO_PENDING_TIMEOUT, 60);
+			libwebsocket_callback_on_writable(Context, Wsi);
 
 		}
 		break;
 	case LWS_CALLBACK_CLOSED_HTTP: 
 		// client went away or 
 		//clean up. 
-		bufferInfo->In.Empty();
-		bufferInfo->Out.Empty(); 
+		BufferInfo->In.Empty();
+		BufferInfo->Out.Empty(); 
 
 		break;
 
@@ -424,15 +418,15 @@ int FNetworkFileServerHttp::CallBack_HTTP(
 	case LWS_CALLBACK_HTTP_WRITEABLE: 
 
 		// get rid of superfluous write  callbacks.
-		if ( bufferInfo == NULL )
+		if ( BufferInfo == NULL )
 			break;
 
 		// we have data o send out.
-		if (bufferInfo->Out.Num())
+		if (BufferInfo->Out.Num())
 		{
-			int SentSize = libwebsocket_write(wsi,(unsigned char*)bufferInfo->Out.GetData(),bufferInfo->Out.Num(), LWS_WRITE_HTTP);
+			int SentSize = libwebsocket_write(Wsi,(unsigned char*)BufferInfo->Out.GetData(),BufferInfo->Out.Num(), LWS_WRITE_HTTP);
 			// get rid of the data that has been sent. 
-			bufferInfo->Out.RemoveAt(0,SentSize); 
+			BufferInfo->Out.RemoveAt(0,SentSize); 
 		}
 
 		break;
diff --git a/Engine/Source/Runtime/NetworkFileSystem/Private/NetworkFileServerHttp.h b/Engine/Source/Runtime/NetworkFileSystem/Private/NetworkFileServerHttp.h
--- a/Engine/Source/Runtime/NetworkFileSystem/Private/NetworkFileServerHttp.h
+++ b/Engine/Source/Runtime/NetworkFileSystem/Private/NetworkFileServerHttp.h
@@ -59,7 +59,7 @@ private:
 	const TArray<ITargetPlatform*> ActiveTargetPlatforms;
 
 	// libwebsocket context. All access to the library happens via this context. 
-	struct libwebsocket_context *context;
+	struct libwebsocket_context *Context;
 
 	// Service Http connections on this thread.
 	FRunnableThread* WorkerThread;
