Made cook on the fly terminate connection earlier if unsupported platform connects.

diff --git a/Engine/Source/Runtime/NetworkFile/Private/NetworkPlatformFile.cpp b/Engine/Source/Runtime/NetworkFile/Private/NetworkPlatformFile.cpp
--- a/Engine/Source/Runtime/NetworkFile/Private/NetworkPlatformFile.cpp
+++ b/Engine/Source/Runtime/NetworkFile/Private/NetworkPlatformFile.cpp
@@ -69,12 +69,16 @@ bool FNetworkPlatformFile::Initialize(IPlatformFile* Inner, const TCHAR* CmdLine
 				// find the correct transport for this ip address 
 				Transport = CreateTransportForHostAddress( HostIpList[HostIpIndex] );
 
+				UE_LOG(LogNetworkPlatformFile, Warning, TEXT("Created transport for %s."), *HostIpList[HostIpIndex]);
+
 				if ( Transport )
 				{
 					bResult = Transport->Initialize( *HostIpList[HostIpIndex] ) && InitializeInternal(Inner, *HostIpList[HostIpIndex]);		
 					if (bResult)
 						break;
 
+					UE_LOG(LogNetworkPlatformFile, Warning, TEXT("Failed to initialize %s."), *HostIpList[HostIpIndex]);
+
 					// try a different host might be a different protocol
 					delete Transport;
 				}
diff --git a/Engine/Source/Runtime/NetworkFileSystem/Private/NetworkFileServer.cpp b/Engine/Source/Runtime/NetworkFileSystem/Private/NetworkFileServer.cpp
--- a/Engine/Source/Runtime/NetworkFileSystem/Private/NetworkFileServer.cpp
+++ b/Engine/Source/Runtime/NetworkFileSystem/Private/NetworkFileServer.cpp
@@ -43,7 +43,13 @@ public:
 				break; 
 			// now process the contents of the payload
 			FBufferArchive Out;
-			FNetworkFileServerClientConnection::ProcessPayload(Payload,Out); 
+			if ( !FNetworkFileServerClientConnection::ProcessPayload(Payload,Out) )
+			{
+				// give the processing of the payload a chance to terminate the connection
+				// failed to process message
+				UE_LOG(LogFileServer, Warning, TEXT("Unable to process payload terminating connection"));
+				break;
+			}
 			if ( !FNFSMessageHeader::WrapAndSendPayload(Out, FSimpleAbstractSocket_FSocket(Socket)))
 				break; 
 		}
diff --git a/Engine/Source/Runtime/NetworkFileSystem/Private/NetworkFileServerConnection.cpp b/Engine/Source/Runtime/NetworkFileSystem/Private/NetworkFileServerConnection.cpp
--- a/Engine/Source/Runtime/NetworkFileSystem/Private/NetworkFileServerConnection.cpp
+++ b/Engine/Source/Runtime/NetworkFileSystem/Private/NetworkFileServerConnection.cpp
@@ -115,8 +115,10 @@ void FNetworkFileServerClientConnection::ConvertServerFilenameToClientFilename(F
 
 static FCriticalSection SocketCriticalSection;
 
-void FNetworkFileServerClientConnection::ProcessPayload(FArchive& Ar,FArchive& Out)
+bool FNetworkFileServerClientConnection::ProcessPayload(FArchive& Ar,FArchive& Out)
 {
+	bool Result = true;
+
 	// first part of the payload is always the command
 	uint32 Cmd;
 	Ar << Cmd;
@@ -210,7 +212,7 @@ void FNetworkFileServerClientConnection::ProcessPayload(FArchive& Ar,FArchive& O
 			break;
 
 		case NFS_Messages::GetFileList:
-			ProcessGetFileList(Ar, Out);
+			Result = ProcessGetFileList(Ar, Out);
 			break;
 
 		case NFS_Messages::Heartbeat:
@@ -246,6 +248,8 @@ void FNetworkFileServerClientConnection::ProcessPayload(FArchive& Ar,FArchive& O
 	}
 
 	UE_LOG(LogFileServer, Verbose, TEXT("Done Processing payload with Cmd %d Total Size sending %d "), Cmd,Out.TotalSize());
+
+	return Result;
 }
 
 
@@ -290,7 +294,7 @@ void FNetworkFileServerClientConnection::ProcessOpenFile( FArchive& In, FArchive
 
 	uint64 HandleId = ++LastHandleId;
 	OpenFiles.Add( HandleId, File );
-
+	
 	Out << HandleId;
 	Out << ServerTimeStamp;
 	Out << ServerFileSize;
@@ -598,7 +602,7 @@ void FNetworkFileServerClientConnection::ProcessToAbsolutePathForWrite( FArchive
 }
 
 
-void FNetworkFileServerClientConnection::ProcessGetFileList( FArchive& In, FArchive& Out )
+bool FNetworkFileServerClientConnection::ProcessGetFileList( FArchive& In, FArchive& Out )
 {
 	// get the list of directories to process
 	TArray<FString> TargetPlatformNames;
@@ -632,13 +636,25 @@ void FNetworkFileServerClientConnection::ProcessGetFileList( FArchive& In, FArch
 			}
 		}
 	}
-	// @todo: if there were some platforms in ActiveTPs, but none matched, we should reject this client
 
-	// if we didn't find one, then just use what was sent
+	// if we didn't find one, reject client and also print some warnings
 	if (ConnectedPlatformName == TEXT(""))
 	{
-		ConnectedPlatformName = TargetPlatformNames[0];
+		// ConnectedPlatformName = TargetPlatformNames[0];
+		// reject client we can't cook for you!
+		UE_LOG(LogFileServer, Warning, TEXT("Unable to find target platform for client, terminating client connection!") );
+
+		for (int32 TPIndex = 0; TPIndex < TargetPlatformNames.Num() && ConnectedPlatformName == TEXT(""); TPIndex++)
+		{
+			UE_LOG(LogFileServer, Warning, TEXT("    Target platforms from client: %s"), *TargetPlatformNames[TPIndex]);
+		}
+		for (int32 ActiveTPIndex = 0; ActiveTPIndex < ActiveTargetPlatforms.Num(); ActiveTPIndex++)
+		{
+			UE_LOG(LogFileServer, Warning, TEXT("    Active target platforms on server: %s"), *ActiveTargetPlatforms[ActiveTPIndex]->PlatformName());
+		}	
+		return false;
 	}
+
 	ConnectedEngineDir = EngineRelativePath;
 	ConnectedGameDir = GameRelativePath;
 
@@ -781,6 +797,7 @@ void FNetworkFileServerClientConnection::ProcessGetFileList( FArchive& In, FArch
 		FixedTimes = FixupSandboxPathsForClient(Sandbox, VisitorForCacheDates.FileTimes, LocalEngineDir, LocalGameDir);
 		Out << FixedTimes;
 	}
+	return true;
 }
 
 
@@ -816,7 +833,7 @@ void FNetworkFileServerClientConnection::PackageFile( FString& Filename, FArchiv
 	{
 		ServerTimeStamp = FDateTime::MinValue(); // if this was a directory, this will make sure it is not confused with a zero byte file
 
-		UE_LOG(LogFileServer, Warning, TEXT("Request for missing file %s."), *Filename);
+		UE_LOG(LogFileServer, Warning, TEXT("Request for missing file %s."), *Filename );
 	}
 	else
 	{
@@ -873,9 +890,9 @@ void FNetworkFileServerClientConnection::ProcessSyncFile( FArchive& In, FArchive
 	// get filename
 	FString Filename;
 	In << Filename;
-
+	
 	ConvertClientFilenameToServerFilename(Filename);
-
+	
 	//FString AbsFile(FString(*Sandbox->ConvertToAbsolutePathForExternalApp(*Filename)).MakeStandardFilename());
 	// ^^ we probably in general want that filename, but for cook on the fly, we want the un-sandboxed name
 
diff --git a/Engine/Source/Runtime/NetworkFileSystem/Private/NetworkFileServerConnection.h b/Engine/Source/Runtime/NetworkFileSystem/Private/NetworkFileServerConnection.h
--- a/Engine/Source/Runtime/NetworkFileSystem/Private/NetworkFileServerConnection.h
+++ b/Engine/Source/Runtime/NetworkFileSystem/Private/NetworkFileServerConnection.h
@@ -32,7 +32,7 @@ public:
 	 * @param Ar An archive containing the payload data.
 	 * @param Out An archive that will contain the processed data.
 	 */
-	void ProcessPayload(FArchive& Ar,FArchive& Out);
+	bool ProcessPayload(FArchive& Ar,FArchive& Out);
 
 	/**
 	 * Gets the client connection's description.
@@ -111,7 +111,7 @@ protected:
 	void ProcessReportLocalFiles(FArchive& In, FArchive& Out);
 
 	/** Walk over a set of directories, and get all files (recursively) in them, along with their timestamps. */
-	void ProcessGetFileList(FArchive& In, FArchive& Out);
+	bool ProcessGetFileList(FArchive& In, FArchive& Out);
 
 	/** Heartbeat. */
 	void ProcessHeartbeat(FArchive& In, FArchive& Out);
diff --git a/Engine/Source/Runtime/NetworkFileSystem/Private/NetworkFileSystemModule.cpp b/Engine/Source/Runtime/NetworkFileSystem/Private/NetworkFileSystemModule.cpp
--- a/Engine/Source/Runtime/NetworkFileSystem/Private/NetworkFileSystemModule.cpp
+++ b/Engine/Source/Runtime/NetworkFileSystem/Private/NetworkFileSystemModule.cpp
@@ -27,11 +27,18 @@ public:
 		// only bother getting the target platforms if there was "-targetplatform" on the commandline, otherwise UnrealFileServer will 
 		// log out some scary sounding, but innocuous logs
 		FString Platforms;
+
+		ITargetPlatformManagerModule& TPM = GetTargetPlatformManagerRef();
+
+		// if we didn't specify a target platform then use the entire target platform list (they could all be possible!)
 		if (FParse::Value(FCommandLine::Get(), TEXT("TARGETPLATFORM="), Platforms))
 		{
-			ITargetPlatformManagerModule& TPM = GetTargetPlatformManagerRef();
 			ActiveTargetPlatforms =  TPM.GetActiveTargetPlatforms();
 		}
+		else
+		{
+			ActiveTargetPlatforms = TPM.GetTargetPlatforms();
+		}
 
 		switch ( Protocol )
 		{
