New connections to network file server from same client cause old connections to terminate.

diff --git a/Engine/Source/Runtime/NetworkFileSystem/Private/NetworkFileServer.cpp b/Engine/Source/Runtime/NetworkFileSystem/Private/NetworkFileServer.cpp
--- a/Engine/Source/Runtime/NetworkFileSystem/Private/NetworkFileServer.cpp
+++ b/Engine/Source/Runtime/NetworkFileSystem/Private/NetworkFileServer.cpp
@@ -37,27 +37,27 @@ public:
 	{
 		while (!StopRequested.GetValue())
 		{
-			// read a header and payload pair
-			FArrayReader Payload; 
-			if (!FNFSMessageHeader::ReceivePayload(Payload, FSimpleAbstractSocket_FSocket(Socket)))
-				break; 
-			// now process the contents of the payload
-			FBufferArchive Out;
-			if ( !FNetworkFileServerClientConnection::ProcessPayload(Payload,Out) )
-			{
-				// give the processing of the payload a chance to terminate the connection
-				// failed to process message
-				UE_LOG(LogFileServer, Warning, TEXT("Unable to process payload terminating connection"));
-				break;
-			}
-			if ( !FNFSMessageHeader::WrapAndSendPayload(Out, FSimpleAbstractSocket_FSocket(Socket)))
-				break; 
-		}
-
-		return true;
+			  // read a header and payload pair
+			  FArrayReader Payload; 
+			  if (!FNFSMessageHeader::ReceivePayload(Payload, FSimpleAbstractSocket_FSocket(Socket)))
+				  break; 
+			  // now process the contents of the payload
+			  FBufferArchive Out;
+			  if ( !FNetworkFileServerClientConnection::ProcessPayload(Payload,Out) )
+			  {
+				  // give the processing of the payload a chance to terminate the connection
+				  // failed to process message
+				  UE_LOG(LogFileServer, Warning, TEXT("Unable to process payload terminating connection"));
+				  break;
+			  }
+			  if ( !FNFSMessageHeader::WrapAndSendPayload(Out, FSimpleAbstractSocket_FSocket(Socket)))
+				  break; 
+		  }
+
+		  return true;
 	}
 
-	virtual void Stop()
+	virtual void Stop() 
 	{
 		StopRequested.Set(true);
 	}
@@ -74,6 +74,11 @@ public:
 		return  (Running.GetValue() != 0); 
 	}
 
+	void GetAddress( FInternetAddr &Addr )
+	{
+		Socket->GetAddress(Addr);
+	}
+
 	~FNetworkFileServerClientConnectionThreaded()
 	{
 		WorkerThread->Kill(true);
@@ -214,6 +219,21 @@ uint32 FNetworkFileServer::Run( )
 
 			if (ClientSocket != NULL)
 			{
+				TSharedPtr<FInternetAddr> Addr = ISocketSubsystem::Get(PLATFORM_SOCKETSUBSYSTEM)->CreateInternetAddr();;
+				ClientSocket->GetAddress(*Addr);
+
+				for ( auto PreviousConnection : Connections )
+				{
+					TSharedPtr<FInternetAddr> PreviousAddr = ISocketSubsystem::Get(PLATFORM_SOCKETSUBSYSTEM)->CreateInternetAddr();;
+					PreviousConnection->GetAddress( *PreviousAddr );
+					if ( *Addr == *PreviousAddr)
+					{
+						// kill hte connection 
+						PreviousConnection->Stop();
+						UE_LOG(LogFileServer, Display, TEXT( "Killing client connection %s because new client connected from same address." ), *PreviousConnection->GetDescription() );
+					}
+				}
+
 				FNetworkFileServerClientConnectionThreaded* Connection = new FNetworkFileServerClientConnectionThreaded(ClientSocket, FileRequestDelegate, RecompileShadersDelegate, ActiveTargetPlatforms);
 				Connections.Add(Connection);
 				UE_LOG(LogFileServer, Display, TEXT( "Client %s connected." ), *Connection->GetDescription() );
