Renamed the http networkfileserver thread so you can tell it appart from the normal one.

diff --git a/Engine/Source/Runtime/NetworkFileSystem/Private/NetworkFileServerHttp.cpp b/Engine/Source/Runtime/NetworkFileSystem/Private/NetworkFileServerHttp.cpp
--- a/Engine/Source/Runtime/NetworkFileSystem/Private/NetworkFileServerHttp.cpp
+++ b/Engine/Source/Runtime/NetworkFileSystem/Private/NetworkFileServerHttp.cpp
@@ -70,7 +70,7 @@ FNetworkFileServerHttp::FNetworkFileServerHttp(
 
 	// spin up the worker thread, this will block till Init has executed on the freshly spinned up thread, Ready will have appropriate value 
 	// set by the end of this function. 
-	WorkerThread = FRunnableThread::Create(this,TEXT("FNetworkFileServer"), 8 * 1024, TPri_AboveNormal);
+	WorkerThread = FRunnableThread::Create(this,TEXT("FNetworkFileServerHttp"), 8 * 1024, TPri_AboveNormal);
 }
 
 
