#HTML5 #UE4  Fix up cook on the fly to use the async Transport APIs.

            - Use a temporary buffer to collect all the unsolicited files and then pass them on as and when requested for single threaded HTML5.

diff --git a/Engine/Source/Runtime/NetworkFile/Private/HTTPTransport.cpp b/Engine/Source/Runtime/NetworkFile/Private/HTTPTransport.cpp
--- a/Engine/Source/Runtime/NetworkFile/Private/HTTPTransport.cpp
+++ b/Engine/Source/Runtime/NetworkFile/Private/HTTPTransport.cpp
@@ -59,6 +59,9 @@ bool FHTTPTransport::Initialize(const TCHAR* InHostIp)
 
 bool FHTTPTransport::SendPayloadAndReceiveResponse(TArray<uint8>& In, TArray<uint8>& Out)
 {
+	RecieveBuffer.Empty();
+	ReadPtr = 0; 
+
 #if !PLATFORM_HTML5
 	
 	if (GIsRequestingExit) // We have already lost HTTP Module. 
@@ -80,7 +83,7 @@ bool FHTTPTransport::SendPayloadAndReceiveResponse(TArray<uint8>& In, TArray<uin
 		TArray<uint8>& Out; 
 	};
 
-	HTTPRequestHandler Handler(Out);
+	HTTPRequestHandler Handler(RecieveBuffer);
 
 	HttpRequest->OnProcessRequestComplete().BindRaw(&Handler,&HTTPRequestHandler::HttpRequestComplete );
 	if ( In.Num() )
@@ -128,31 +131,58 @@ bool FHTTPTransport::SendPayloadAndReceiveResponse(TArray<uint8>& In, TArray<uin
 	unsigned char *OutData = NULL; 
 	unsigned int OutSize= 0; 
 
-#if PLATFORM_HTML5_WIN32
+	bool RetVal = true; 
 
+#if PLATFORM_HTML5_WIN32
+	RetVal = HTML5Win32::NFSHttp::SendPayLoadAndRecieve(Ar.GetData(), Ar.Num(), &OutData, OutSize);
+#endif 
+#if PLATFORM_HTML5_BROWSER
+	UE_SendAndRecievePayLoad(TCHAR_TO_ANSI(Url),(char*)Ar.GetData(),Ar.Num(),(char**)&OutData,(int*)&OutSize);
+#endif 
 
-	bool RetVal = HTML5Win32::NFSHttp::SendPayLoadAndRecieve(Ar.GetData(), Ar.Num(), &OutData, OutSize);
-	if (OutSize)
+	if (!Ar.Num())
 	{
-		Out.Append(OutData,OutSize); 
-		free (OutData); 
+		uint32 Size = OutSize;
+		uint32 Marker = 0xDeadBeef; 
+		RecieveBuffer.Append((uint8*)&Marker,sizeof(uint32));
+		RecieveBuffer.Append((uint8*)&Size,sizeof(uint32));
 	}
 
-	return RetVal; 
-#endif 
 
-#if PLATFORM_HTML5_BROWSER
-	UE_SendAndRecievePayLoad(TCHAR_TO_ANSI(Url),(char*)Ar.GetData(),Ar.Num(),(char**)&OutData,(int*)&OutSize);
 	if (OutSize)
 	{
-		Out.Append(OutData,OutSize); 
+		RecieveBuffer.Append(OutData,OutSize); 
+
+#if PLATFORM_HTML5_WIN32
+		free (OutData); 
+#endif 
+#if PLATFORM_HTML5_BROWSER
 		// don't go through the Unreal Memory system. 
-		::free (OutData); 
+		::free(OutData);
+#endif 
+
 	}
-	return true; 
+
+	return RetVal & ReceiveResponse(Out); 
 #endif 
-#endif
+}
+
+bool FHTTPTransport::ReceiveResponse(TArray<uint8> &Out)
+{
+	// Read one Packet from Receive Buffer.
+	// read the size. 
+
+	uint32 Marker = *(uint32*)(RecieveBuffer.GetData() + ReadPtr); 
+	uint32 Size = *(uint32*)(RecieveBuffer.GetData() + ReadPtr + sizeof(uint32));
+
+	// make sure we have the right amount of data available in the buffer. 
+	check( (ReadPtr + Size + 2*sizeof(uint32)) <= RecieveBuffer.Num());
+
+	Out.Append(RecieveBuffer.GetData() + ReadPtr + 2*sizeof(uint32),Size);
+
+	ReadPtr += 2*sizeof(uint32) + Size; 
 
+	return true;
 }
 
 #endif 
\ No newline at end of file
diff --git a/Engine/Source/Runtime/NetworkFile/Private/HTTPTransport.h b/Engine/Source/Runtime/NetworkFile/Private/HTTPTransport.h
--- a/Engine/Source/Runtime/NetworkFile/Private/HTTPTransport.h
+++ b/Engine/Source/Runtime/NetworkFile/Private/HTTPTransport.h
@@ -17,7 +17,7 @@ public:
 	// ITransport Interface.
 	virtual bool Initialize(const TCHAR* HostIp) override;
 	virtual bool SendPayloadAndReceiveResponse(TArray<uint8>& In, TArray<uint8>& Out) override;
-	virtual bool ReceiveResponse(TArray<uint8> &Out) override { return false; } // unsupported for now :(
+	virtual bool ReceiveResponse(TArray<uint8> &Out) override; 
 
 private: 
 
@@ -28,5 +28,8 @@ private:
 	FGuid Guid; 
 	TCHAR Url[1048];
 
+	TArray<uint8> RecieveBuffer; 
+	uint32 ReadPtr; 
+
 };
 #endif
\ No newline at end of file
diff --git a/Engine/Source/Runtime/NetworkFileSystem/Private/NetworkFileServerConnection.h b/Engine/Source/Runtime/NetworkFileSystem/Private/NetworkFileServerConnection.h
--- a/Engine/Source/Runtime/NetworkFileSystem/Private/NetworkFileServerConnection.h
+++ b/Engine/Source/Runtime/NetworkFileSystem/Private/NetworkFileServerConnection.h
@@ -148,7 +148,7 @@ protected:
 	void ProcessSyncFile( FArchive& In, FArchive& Out );
 
 
-	virtual bool SendPayload( TArray<uint8> &Out ) { return false; }
+	virtual bool SendPayload( TArray<uint8> &Out ) = 0; 
 	
 private:
 
diff --git a/Engine/Source/Runtime/NetworkFileSystem/Private/NetworkFileServerHttp.cpp b/Engine/Source/Runtime/NetworkFileSystem/Private/NetworkFileServerHttp.cpp
--- a/Engine/Source/Runtime/NetworkFileSystem/Private/NetworkFileServerHttp.cpp
+++ b/Engine/Source/Runtime/NetworkFileSystem/Private/NetworkFileServerHttp.cpp
@@ -8,6 +8,44 @@
 #include "AllowWindowsPlatformTypes.h"
 	#include "libwebsockets.h"
 #include "HideWindowsPlatformTypes.h"
+
+
+
+class FNetworkFileServerClientConnectionHTTP : public FNetworkFileServerClientConnection 
+{
+
+public: 
+	FNetworkFileServerClientConnectionHTTP(const FFileRequestDelegate& InFileRequestDelegate, 
+		const FRecompileShadersDelegate& InRecompileShadersDelegate, const TArray<ITargetPlatform*>& InActiveTargetPlatforms )
+		:  FNetworkFileServerClientConnection( InFileRequestDelegate,InRecompileShadersDelegate,InActiveTargetPlatforms)
+	{
+	} 
+
+	bool SendPayload( TArray<uint8> &Out )
+	{
+		// Make Boundaries between payloads, add a visual marker for easier debugging.
+
+		uint32 Marker = 0xDeadBeef; 
+		uint32 Size = Out.Num(); 
+
+		OutBuffer.Append((uint8*)&Marker,sizeof(uint32));
+		OutBuffer.Append((uint8*)&Size,sizeof(uint32));
+		OutBuffer.Append(Out);
+
+		return true; 
+	}
+
+private:  
+
+	TArray<uint8>& GetOutBuffer() { return OutBuffer; }
+	void ResetBuffer() { OutBuffer.Reset(); }
+
+	TArray<uint8> OutBuffer; 
+
+    friend class FNetworkFileServerHttp;  
+};
+
+
 ////////////////////////////////////////////////////////////////////////// 
 // LibWebsockets specific structs. 
 
@@ -207,13 +245,13 @@ FNetworkFileServerHttp::~FNetworkFileServerHttp()
 	check( Context == NULL ); 
 }
 
-FNetworkFileServerClientConnection* FNetworkFileServerHttp::CreateNewConnection()
+FNetworkFileServerClientConnectionHTTP* FNetworkFileServerHttp::CreateNewConnection()
 {
-	return new FNetworkFileServerClientConnection (FileRequestDelegate,RecompileShadersDelegate,ActiveTargetPlatforms); 
+	return new FNetworkFileServerClientConnectionHTTP(FileRequestDelegate,RecompileShadersDelegate,ActiveTargetPlatforms); 
 }
 
 // Have a similar process function for the normal tcp connection. 
-void FNetworkFileServerHttp::Process(FArchive& In, FArchive&Out, FNetworkFileServerHttp* Server)
+void FNetworkFileServerHttp::Process(FArchive& In, TArray<uint8>&Out, FNetworkFileServerHttp* Server)
 {
 	int loops = 0; 
 	while(!In.AtEnd())
@@ -225,7 +263,7 @@ void FNetworkFileServerHttp::Process(FArchive& In, FArchive&Out, FNetworkFileSer
 
 		UE_LOG(LogFileServer, Log, TEXT("Recieved GUID %s"), *ClientGuid.ToString());
 
-		FNetworkFileServerClientConnection* Connection = NULL; 
+		FNetworkFileServerClientConnectionHTTP* Connection = NULL; 
 		if (Server->RequestHandlers.Contains(ClientGuid))
 		{
 			UE_LOG(LogFileServer, Log, TEXT("Picking up an existing handler" ));
@@ -238,7 +276,9 @@ void FNetworkFileServerHttp::Process(FArchive& In, FArchive&Out, FNetworkFileSer
 			Server->RequestHandlers.Add(ClientGuid,Connection);
 		}
 
-		// Connection->ProcessPayload(In, Out);
+	    Connection->ProcessPayload(In);
+		Out.Append(Connection->GetOutBuffer());
+		Connection->ResetBuffer(); 
 	}
 }
 
@@ -371,7 +411,7 @@ int FNetworkFileServerHttp::CallBack_HTTP(
 			// create archives and process them. 
 			UE_LOG(LogFileServer, Log, TEXT("Incoming HTTP total size  %d"), BufferInfo->In.Num());
 			FMemoryReader Reader(BufferInfo->In);
-			FBufferArchive Writer;
+			TArray<uint8> Writer;
 
 			FNetworkFileServerHttp::Process(Reader,Writer,Server);
 
@@ -434,4 +474,4 @@ int FNetworkFileServerHttp::CallBack_HTTP(
 	return 0; 
 }
 
-#endif
\ No newline at end of file
+#endif
diff --git a/Engine/Source/Runtime/NetworkFileSystem/Private/NetworkFileServerHttp.h b/Engine/Source/Runtime/NetworkFileSystem/Private/NetworkFileServerHttp.h
--- a/Engine/Source/Runtime/NetworkFileSystem/Private/NetworkFileServerHttp.h
+++ b/Engine/Source/Runtime/NetworkFileSystem/Private/NetworkFileServerHttp.h
@@ -43,11 +43,11 @@ private:
 	virtual void Stop();
 	virtual void Exit();
 
-	static void Process (FArchive&, FArchive&, FNetworkFileServerHttp* );
+	static void Process (FArchive&, TArray<uint8>&, FNetworkFileServerHttp* );
 
 
 	// factory method for creating a new Client Connection. 
-	FNetworkFileServerClientConnection* CreateNewConnection(); 
+	class FNetworkFileServerClientConnectionHTTP* CreateNewConnection(); 
 
 	// Holds a delegate to be invoked on every sync request.
 	FFileRequestDelegate FileRequestDelegate;
@@ -74,7 +74,7 @@ private:
 	FThreadSafeCounter Ready; 
 
 	// Clients being served.
-	TMap< FGuid, FNetworkFileServerClientConnection* > RequestHandlers; 
+	TMap< FGuid, FNetworkFileServerClientConnectionHTTP* > RequestHandlers; 
 };
 
 #endif
\ No newline at end of file
