- Fixed UnrealFileServer after recent changes

diff --git a/Engine/Source/Developer/TargetPlatform/TargetPlatform.Build.cs b/Engine/Source/Developer/TargetPlatform/TargetPlatform.Build.cs
--- a/Engine/Source/Developer/TargetPlatform/TargetPlatform.Build.cs
+++ b/Engine/Source/Developer/TargetPlatform/TargetPlatform.Build.cs
@@ -10,7 +10,8 @@ public class TargetPlatform : ModuleRules
 		PrivateDependencyModuleNames.Add("Core");		
         PublicDependencyModuleNames.Add("DesktopPlatform");
 
-		if (!UEBuildConfiguration.bBuildRequiresCookedData)
+		// no need for all these modules if the program doesn't want developer tools at all (like UnrealFileServer)
+		if (!UEBuildConfiguration.bBuildRequiresCookedData && UEBuildConfiguration.bBuildDeveloperTools)
 		{
             // these are needed by multiple platform specific target platforms, so we make sure they are built with the base editor
             DynamicallyLoadedModuleNames.Add("ShaderPreprocessor");
diff --git a/Engine/Source/Editor/UnrealEd/Private/Commandlets/FileServerCommandlet.cpp b/Engine/Source/Editor/UnrealEd/Private/Commandlets/FileServerCommandlet.cpp
--- a/Engine/Source/Editor/UnrealEd/Private/Commandlets/FileServerCommandlet.cpp
+++ b/Engine/Source/Editor/UnrealEd/Private/Commandlets/FileServerCommandlet.cpp
@@ -58,7 +58,7 @@ int32 UFileServerCommandlet::Main( const FString& Params )
 
 	// start the listening thread
 	INetworkFileServer* NetworkFileServer = FModuleManager::LoadModuleChecked<INetworkFileSystemModule>("NetworkFileSystem")
-		.CreateNetworkFileServer(InstanceId.IsValid() ? 0 : -1);
+		.CreateNetworkFileServer(true, InstanceId.IsValid() ? 0 : -1);
 
 	TArray<TSharedPtr<FInternetAddr> > AddressList;
 
diff --git a/Engine/Source/Editor/UnrealEd/Private/CookOnTheFlyServer.cpp b/Engine/Source/Editor/UnrealEd/Private/CookOnTheFlyServer.cpp
--- a/Engine/Source/Editor/UnrealEd/Private/CookOnTheFlyServer.cpp
+++ b/Engine/Source/Editor/UnrealEd/Private/CookOnTheFlyServer.cpp
@@ -303,14 +303,14 @@ bool UCookOnTheFlyServer::StartNetworkFileServer( const bool BindAnyPort )
 	FRecompileShadersDelegate RecompileShadersDelegate(FRecompileShadersDelegate::CreateUObject(this, &UCookOnTheFlyServer::HandleNetworkFileServerRecompileShaders));
 
 	INetworkFileServer *TcpFileServer = FModuleManager::LoadModuleChecked<INetworkFileSystemModule>("NetworkFileSystem")
-		.CreateNetworkFileServer(BindAnyPort ? 0 : -1, &FileRequestDelegate, &RecompileShadersDelegate, ENetworkFileServerProtocol::NFSP_Tcp);
+		.CreateNetworkFileServer(true, BindAnyPort ? 0 : -1, &FileRequestDelegate, &RecompileShadersDelegate, ENetworkFileServerProtocol::NFSP_Tcp);
 	if ( TcpFileServer )
 	{
 		NetworkFileServers.Add(TcpFileServer);
 	}
 
 	INetworkFileServer *HttpFileServer = FModuleManager::LoadModuleChecked<INetworkFileSystemModule>("NetworkFileSystem")
-		.CreateNetworkFileServer(BindAnyPort ? 0 : -1, &FileRequestDelegate, &RecompileShadersDelegate, ENetworkFileServerProtocol::NFSP_Http);
+		.CreateNetworkFileServer(true, BindAnyPort ? 0 : -1, &FileRequestDelegate, &RecompileShadersDelegate, ENetworkFileServerProtocol::NFSP_Http);
 	if ( HttpFileServer )
 	{
 		NetworkFileServers.Add( HttpFileServer );
diff --git a/Engine/Source/Programs/UnrealFileServer/Private/UnrealFileServer.cpp b/Engine/Source/Programs/UnrealFileServer/Private/UnrealFileServer.cpp
--- a/Engine/Source/Programs/UnrealFileServer/Private/UnrealFileServer.cpp
+++ b/Engine/Source/Programs/UnrealFileServer/Private/UnrealFileServer.cpp
@@ -34,7 +34,7 @@ INT32_MAIN_INT32_ARGC_TCHAR_ARGV()
 
 	// start the listening thread
 	INetworkFileServer* NetworkFileServer = FModuleManager::LoadModuleChecked<INetworkFileSystemModule>("NetworkFileSystem")
-		.CreateNetworkFileServer();
+		.CreateNetworkFileServer(false);
 
 	// loop while the server does the rest
 	double LastTime = FPlatformTime::Seconds();
diff --git a/Engine/Source/Programs/UnrealFileServer/UnrealFileServer.Build.cs b/Engine/Source/Programs/UnrealFileServer/UnrealFileServer.Build.cs
--- a/Engine/Source/Programs/UnrealFileServer/UnrealFileServer.Build.cs
+++ b/Engine/Source/Programs/UnrealFileServer/UnrealFileServer.Build.cs
@@ -19,7 +19,6 @@ public class UnrealFileServer : ModuleRules
                 "SandboxFile",
 				"Sockets",
 				"Projects",
-				"Messaging",
            }
 		);
 	}
diff --git a/Engine/Source/Programs/UnrealFileServer/UnrealFileServer.Target.cs b/Engine/Source/Programs/UnrealFileServer/UnrealFileServer.Target.cs
--- a/Engine/Source/Programs/UnrealFileServer/UnrealFileServer.Target.cs
+++ b/Engine/Source/Programs/UnrealFileServer/UnrealFileServer.Target.cs
@@ -52,6 +52,7 @@ public class UnrealFileServerTarget : TargetRules
 		// Currently this app is not linking against the engine, so we'll compile out references from Core to the rest of the engine
 		UEBuildConfiguration.bCompileAgainstEngine = false;
 		UEBuildConfiguration.bCompileAgainstCoreUObject = false;
+		UEBuildConfiguration.bBuildDeveloperTools = false;
 
 		// UnrealHeaderTool is a console application, not a Windows app (sets entry point to main(), instead of WinMain())
 		OutLinkEnvironmentConfiguration.bIsBuildingConsoleApplication = true;
diff --git a/Engine/Source/Runtime/NetworkFileSystem/Private/NetworkFileServerConnection.cpp b/Engine/Source/Runtime/NetworkFileSystem/Private/NetworkFileServerConnection.cpp
--- a/Engine/Source/Runtime/NetworkFileSystem/Private/NetworkFileServerConnection.cpp
+++ b/Engine/Source/Runtime/NetworkFileSystem/Private/NetworkFileServerConnection.cpp
@@ -643,6 +643,7 @@ bool FNetworkFileServerClientConnection::ProcessGetFileList( FArchive& In, FArch
 	In << bIsStreamingRequest;
 
 	ConnectedPlatformName = TEXT("");
+
 	// figure out the best matching target platform for the set of valid ones
 	for (int32 TPIndex = 0; TPIndex < TargetPlatformNames.Num() && ConnectedPlatformName == TEXT(""); TPIndex++)
 	{
@@ -663,19 +664,27 @@ bool FNetworkFileServerClientConnection::ProcessGetFileList( FArchive& In, FArch
 	// if we didn't find one, reject client and also print some warnings
 	if (ConnectedPlatformName == TEXT(""))
 	{
-		// ConnectedPlatformName = TargetPlatformNames[0];
-		// reject client we can't cook for you!
-		UE_LOG(LogFileServer, Warning, TEXT("Unable to find target platform for client, terminating client connection!") );
-
-		for (int32 TPIndex = 0; TPIndex < TargetPlatformNames.Num() && ConnectedPlatformName == TEXT(""); TPIndex++)
+		// if we didn't find one (and this is a dumb server - no active platforms), then just use what was sent
+		if (ActiveTargetPlatforms.Num() == 0)
 		{
-			UE_LOG(LogFileServer, Warning, TEXT("    Target platforms from client: %s"), *TargetPlatformNames[TPIndex]);
+			ConnectedPlatformName = TargetPlatformNames[0];
 		}
-		for (int32 ActiveTPIndex = 0; ActiveTPIndex < ActiveTargetPlatforms.Num(); ActiveTPIndex++)
+		// we only need to care about validating the connected platform if there are active targetplatforms
+		else
 		{
-			UE_LOG(LogFileServer, Warning, TEXT("    Active target platforms on server: %s"), *ActiveTargetPlatforms[ActiveTPIndex]->PlatformName());
-		}	
-		return false;
+			// reject client we can't cook/compile shaders for you!
+			UE_LOG(LogFileServer, Warning, TEXT("Unable to find target platform for client, terminating client connection!"));
+
+			for (int32 TPIndex = 0; TPIndex < TargetPlatformNames.Num() && ConnectedPlatformName == TEXT(""); TPIndex++)
+			{
+				UE_LOG(LogFileServer, Warning, TEXT("    Target platforms from client: %s"), *TargetPlatformNames[TPIndex]);
+			}
+			for (int32 ActiveTPIndex = 0; ActiveTPIndex < ActiveTargetPlatforms.Num(); ActiveTPIndex++)
+			{
+				UE_LOG(LogFileServer, Warning, TEXT("    Active target platforms on server: %s"), *ActiveTargetPlatforms[ActiveTPIndex]->PlatformName());
+			}
+			return false;
+		}
 	}
 
 	ConnectedEngineDir = EngineRelativePath;
diff --git a/Engine/Source/Runtime/NetworkFileSystem/Private/NetworkFileSystemModule.cpp b/Engine/Source/Runtime/NetworkFileSystem/Private/NetworkFileSystemModule.cpp
--- a/Engine/Source/Runtime/NetworkFileSystem/Private/NetworkFileSystemModule.cpp
+++ b/Engine/Source/Runtime/NetworkFileSystem/Private/NetworkFileSystemModule.cpp
@@ -20,24 +20,23 @@ public:
 
 	// INetworkFileSystemModule interface
 
-	virtual INetworkFileServer* CreateNetworkFileServer( int32 Port, const FFileRequestDelegate* InFileRequestDelegate, const FRecompileShadersDelegate* InRecompileShadersDelegate, const ENetworkFileServerProtocol Protocol ) const override
+	virtual INetworkFileServer* CreateNetworkFileServer( bool bLoadTargetPlatforms, int32 Port, const FFileRequestDelegate* InFileRequestDelegate, const FRecompileShadersDelegate* InRecompileShadersDelegate, const ENetworkFileServerProtocol Protocol ) const override
 	{
 		TArray<ITargetPlatform*> ActiveTargetPlatforms;
-
-		// only bother getting the target platforms if there was "-targetplatform" on the commandline, otherwise UnrealFileServer will 
-		// log out some scary sounding, but innocuous logs
-		FString Platforms;
-
-		ITargetPlatformManagerModule& TPM = GetTargetPlatformManagerRef();
-
-		// if we didn't specify a target platform then use the entire target platform list (they could all be possible!)
-		if (FParse::Value(FCommandLine::Get(), TEXT("TARGETPLATFORM="), Platforms))
-		{
-			ActiveTargetPlatforms =  TPM.GetActiveTargetPlatforms();
-		}
-		else
+		if (bLoadTargetPlatforms)
 		{
-			ActiveTargetPlatforms = TPM.GetTargetPlatforms();
+			ITargetPlatformManagerModule& TPM = GetTargetPlatformManagerRef();
+
+			// if we didn't specify a target platform then use the entire target platform list (they could all be possible!)
+			FString Platforms;
+			if (FParse::Value(FCommandLine::Get(), TEXT("TARGETPLATFORM="), Platforms))
+			{
+				ActiveTargetPlatforms =  TPM.GetActiveTargetPlatforms();
+			}
+			else
+			{
+				ActiveTargetPlatforms = TPM.GetTargetPlatforms();
+			}
 		}
 
 		switch ( Protocol )
diff --git a/Engine/Source/Runtime/NetworkFileSystem/Public/Interfaces/INetworkFileSystemModule.h b/Engine/Source/Runtime/NetworkFileSystem/Public/Interfaces/INetworkFileSystemModule.h
--- a/Engine/Source/Runtime/NetworkFileSystem/Public/Interfaces/INetworkFileSystemModule.h
+++ b/Engine/Source/Runtime/NetworkFileSystem/Public/Interfaces/INetworkFileSystemModule.h
@@ -71,7 +71,7 @@ public:
 	 *
 	 * @return The new file server, or nullptr if creation failed.
 	 */
-	virtual INetworkFileServer* CreateNetworkFileServer( int32 Port = -1, const FFileRequestDelegate* InFileRequestDelegate = nullptr, const FRecompileShadersDelegate* InRecompileShadersDelegate = nullptr, const ENetworkFileServerProtocol Protocol = NFSP_Tcp ) const = 0;
+	virtual INetworkFileServer* CreateNetworkFileServer( bool bLoadTargetPlatforms, int32 Port = -1, const FFileRequestDelegate* InFileRequestDelegate = nullptr, const FRecompileShadersDelegate* InRecompileShadersDelegate = nullptr, const ENetworkFileServerProtocol Protocol = NFSP_Tcp ) const = 0;
 
 public:
 
