- Minor cleanup to file server platform checking code

diff --git a/Engine/Source/Runtime/NetworkFileSystem/Private/NetworkFileServerConnection.cpp b/Engine/Source/Runtime/NetworkFileSystem/Private/NetworkFileServerConnection.cpp
--- a/Engine/Source/Runtime/NetworkFileSystem/Private/NetworkFileServerConnection.cpp
+++ b/Engine/Source/Runtime/NetworkFileSystem/Private/NetworkFileServerConnection.cpp
@@ -644,46 +644,46 @@ bool FNetworkFileServerClientConnection::ProcessGetFileList( FArchive& In, FArch
 
 	ConnectedPlatformName = TEXT("");
 
-	// figure out the best matching target platform for the set of valid ones
-	for (int32 TPIndex = 0; TPIndex < TargetPlatformNames.Num() && ConnectedPlatformName == TEXT(""); TPIndex++)
+	// if we didn't find one (and this is a dumb server - no active platforms), then just use what was sent
+	if (ActiveTargetPlatforms.Num() == 0)
 	{
-		UE_LOG(LogFileServer, Display, TEXT("    Possible Target Platform from client: %s"), *TargetPlatformNames[TPIndex]);
-
-		// look for a matching target platform
-		for (int32 ActiveTPIndex = 0; ActiveTPIndex < ActiveTargetPlatforms.Num(); ActiveTPIndex++)
-		{
-			UE_LOG(LogFileServer, Display, TEXT("   Checking against: %s"), *ActiveTargetPlatforms[ActiveTPIndex]->PlatformName());
-			if (ActiveTargetPlatforms[ActiveTPIndex]->PlatformName() == TargetPlatformNames[TPIndex])
-			{
-				ConnectedPlatformName = ActiveTargetPlatforms[ActiveTPIndex]->PlatformName();
-				break;
-			}
-		}
+		ConnectedPlatformName = TargetPlatformNames[0];
 	}
-
-	// if we didn't find one, reject client and also print some warnings
-	if (ConnectedPlatformName == TEXT(""))
+	// we only need to care about validating the connected platform if there are active targetplatforms
+	else
 	{
-		// if we didn't find one (and this is a dumb server - no active platforms), then just use what was sent
-		if (ActiveTargetPlatforms.Num() == 0)
+		// figure out the best matching target platform for the set of valid ones
+		for (int32 TPIndex = 0; TPIndex < TargetPlatformNames.Num() && ConnectedPlatformName == TEXT(""); TPIndex++)
 		{
-			ConnectedPlatformName = TargetPlatformNames[0];
-		}
-		// we only need to care about validating the connected platform if there are active targetplatforms
-		else
-		{
-			// reject client we can't cook/compile shaders for you!
-			UE_LOG(LogFileServer, Warning, TEXT("Unable to find target platform for client, terminating client connection!"));
+			UE_LOG(LogFileServer, Display, TEXT("    Possible Target Platform from client: %s"), *TargetPlatformNames[TPIndex]);
 
-			for (int32 TPIndex = 0; TPIndex < TargetPlatformNames.Num() && ConnectedPlatformName == TEXT(""); TPIndex++)
+			// look for a matching target platform
+			for (int32 ActiveTPIndex = 0; ActiveTPIndex < ActiveTargetPlatforms.Num(); ActiveTPIndex++)
 			{
-				UE_LOG(LogFileServer, Warning, TEXT("    Target platforms from client: %s"), *TargetPlatformNames[TPIndex]);
+				UE_LOG(LogFileServer, Display, TEXT("   Checking against: %s"), *ActiveTargetPlatforms[ActiveTPIndex]->PlatformName());
+				if (ActiveTargetPlatforms[ActiveTPIndex]->PlatformName() == TargetPlatformNames[TPIndex])
+				{
+					ConnectedPlatformName = ActiveTargetPlatforms[ActiveTPIndex]->PlatformName();
+					break;
+				}
 			}
-			for (int32 ActiveTPIndex = 0; ActiveTPIndex < ActiveTargetPlatforms.Num(); ActiveTPIndex++)
+
+			// if we didn't find one, reject client and also print some warnings
+			if (ConnectedPlatformName == TEXT(""))
 			{
-				UE_LOG(LogFileServer, Warning, TEXT("    Active target platforms on server: %s"), *ActiveTargetPlatforms[ActiveTPIndex]->PlatformName());
+				// reject client we can't cook/compile shaders for you!
+				UE_LOG(LogFileServer, Warning, TEXT("Unable to find target platform for client, terminating client connection!"));
+
+				for (int32 TPIndex = 0; TPIndex < TargetPlatformNames.Num() && ConnectedPlatformName == TEXT(""); TPIndex++)
+				{
+					UE_LOG(LogFileServer, Warning, TEXT("    Target platforms from client: %s"), *TargetPlatformNames[TPIndex]);
+				}
+				for (int32 ActiveTPIndex = 0; ActiveTPIndex < ActiveTargetPlatforms.Num(); ActiveTPIndex++)
+				{
+					UE_LOG(LogFileServer, Warning, TEXT("    Active target platforms on server: %s"), *ActiveTargetPlatforms[ActiveTPIndex]->PlatformName());
+				}
+				return false;
 			}
-			return false;
 		}
 	}
 
