- Fixed a bug with COTF server if the first target platform from client isn't found in possible server targets (now happening with ASTC devices)

diff --git a/Engine/Source/Runtime/NetworkFileSystem/Private/NetworkFileServerConnection.cpp b/Engine/Source/Runtime/NetworkFileSystem/Private/NetworkFileServerConnection.cpp
--- a/Engine/Source/Runtime/NetworkFileSystem/Private/NetworkFileServerConnection.cpp
+++ b/Engine/Source/Runtime/NetworkFileSystem/Private/NetworkFileServerConnection.cpp
@@ -667,23 +667,23 @@ bool FNetworkFileServerClientConnection::ProcessGetFileList( FArchive& In, FArch
 					break;
 				}
 			}
+		}
 
-			// if we didn't find one, reject client and also print some warnings
-			if (ConnectedPlatformName == TEXT(""))
-			{
-				// reject client we can't cook/compile shaders for you!
-				UE_LOG(LogFileServer, Warning, TEXT("Unable to find target platform for client, terminating client connection!"));
+		// if we didn't find one, reject client and also print some warnings
+		if (ConnectedPlatformName == TEXT(""))
+		{
+			// reject client we can't cook/compile shaders for you!
+			UE_LOG(LogFileServer, Warning, TEXT("Unable to find target platform for client, terminating client connection!"));
 
-				for (int32 TPIndex = 0; TPIndex < TargetPlatformNames.Num() && ConnectedPlatformName == TEXT(""); TPIndex++)
-				{
-					UE_LOG(LogFileServer, Warning, TEXT("    Target platforms from client: %s"), *TargetPlatformNames[TPIndex]);
-				}
-				for (int32 ActiveTPIndex = 0; ActiveTPIndex < ActiveTargetPlatforms.Num(); ActiveTPIndex++)
-				{
-					UE_LOG(LogFileServer, Warning, TEXT("    Active target platforms on server: %s"), *ActiveTargetPlatforms[ActiveTPIndex]->PlatformName());
-				}
-				return false;
+			for (int32 TPIndex = 0; TPIndex < TargetPlatformNames.Num() && ConnectedPlatformName == TEXT(""); TPIndex++)
+			{
+				UE_LOG(LogFileServer, Warning, TEXT("    Target platforms from client: %s"), *TargetPlatformNames[TPIndex]);
+			}
+			for (int32 ActiveTPIndex = 0; ActiveTPIndex < ActiveTargetPlatforms.Num(); ActiveTPIndex++)
+			{
+				UE_LOG(LogFileServer, Warning, TEXT("    Active target platforms on server: %s"), *ActiveTargetPlatforms[ActiveTPIndex]->PlatformName());
 			}
+			return false;
 		}
 	}
 
