Increased stack size for the networkfile server thread because of stack overflow issues.

diff --git a/Engine/Source/Runtime/NetworkFileSystem/Private/NetworkFileServer.cpp b/Engine/Source/Runtime/NetworkFileSystem/Private/NetworkFileServer.cpp
--- a/Engine/Source/Runtime/NetworkFileSystem/Private/NetworkFileServer.cpp
+++ b/Engine/Source/Runtime/NetworkFileSystem/Private/NetworkFileServer.cpp
@@ -22,7 +22,7 @@ public:
 			// this thread needs more space in debug builds as it tries to log messages and such
 			const static uint32 NetworkFileServerThreadSize = 2 * 1024 * 1024; 
 #else
-			const static uint32 NetworkFileServerThreadSize = 8 * 1024; 
+			const static uint32 NetworkFileServerThreadSize = 1 * 1024 * 1024; 
 #endif
 			WorkerThread = FRunnableThread::Create(this, TEXT("FNetworkFileServerClientConnection"), NetworkFileServerThreadSize, TPri_AboveNormal);
 	}
