Reduce cookonthefly getfilelist times for games which also have a staged build.

diff --git a/Engine/Source/Runtime/NetworkFileSystem/Private/NetworkFileServerConnection.cpp b/Engine/Source/Runtime/NetworkFileSystem/Private/NetworkFileServerConnection.cpp
--- a/Engine/Source/Runtime/NetworkFileSystem/Private/NetworkFileServerConnection.cpp
+++ b/Engine/Source/Runtime/NetworkFileSystem/Private/NetworkFileServerConnection.cpp
@@ -774,6 +774,7 @@ bool FNetworkFileServerClientConnection::ProcessGetFileList( FArchive& In, FArch
 		DirectoriesToSkip.Add(FString(RootDirectories[DirIndex] / TEXT("Saved/Sandboxes")));
 		DirectoriesToSkip.Add(FString(RootDirectories[DirIndex] / TEXT("Saved/Cooked")));
 		DirectoriesToSkip.Add(FString(RootDirectories[DirIndex] / TEXT("Saved/ShaderDebugInfo")));
+		DirectoriesToSkip.Add(FString(RootDirectories[DirIndex] / TEXT("Saved/StagedBuilds")));
 		DirectoriesToSkip.Add(FString(RootDirectories[DirIndex] / TEXT("Intermediate")));
 		DirectoriesToSkip.Add(FString(RootDirectories[DirIndex] / TEXT("Documentation")));
 		DirectoriesToSkip.Add(FString(RootDirectories[DirIndex] / TEXT("Extras")));
