Do PS4 Path shortening at staging time.

diff --git a/Engine/Source/Developer/TargetPlatform/Public/Common/TargetPlatformBase.h b/Engine/Source/Developer/TargetPlatform/Public/Common/TargetPlatformBase.h
--- a/Engine/Source/Developer/TargetPlatform/Public/Common/TargetPlatformBase.h
+++ b/Engine/Source/Developer/TargetPlatform/Public/Common/TargetPlatformBase.h
@@ -196,6 +196,11 @@ public:
 		return 0.0f;
 	}
 
+	virtual bool SendLowerCaseFilePaths() const override
+	{
+		return false;
+	}
+
 protected:
 
 	FTargetPlatformBase(const PlatformInfo::FPlatformInfo *const InPlatformInfo)
diff --git a/Engine/Source/Developer/TargetPlatform/Public/Interfaces/ITargetPlatform.h b/Engine/Source/Developer/TargetPlatform/Public/Interfaces/ITargetPlatform.h
--- a/Engine/Source/Developer/TargetPlatform/Public/Interfaces/ITargetPlatform.h
+++ b/Engine/Source/Developer/TargetPlatform/Public/Interfaces/ITargetPlatform.h
@@ -367,6 +367,11 @@ public:
 	 */
 	virtual float GetVariantPriority() const = 0;
 
+	/** 
+	 * Whether or not to send all lower-case filepaths when connecting over a fileserver connection. 
+	 */
+	virtual bool SendLowerCaseFilePaths() const = 0;
+
 public:
 
 	/**
diff --git a/Engine/Source/Runtime/NetworkFile/Private/NetworkPlatformFile.cpp b/Engine/Source/Runtime/NetworkFile/Private/NetworkPlatformFile.cpp
--- a/Engine/Source/Runtime/NetworkFile/Private/NetworkPlatformFile.cpp
+++ b/Engine/Source/Runtime/NetworkFile/Private/NetworkPlatformFile.cpp
@@ -254,6 +254,9 @@ void FNetworkPlatformFile::InitializeAfterSetActive()
 				ConvertServerFilenameToClientFilename(ServerFile);
 
 				// Set it in the visitor file times list
+				// If there is any pathing difference (casing, or relative path, or whatever) between the server's filelist and the results
+				// of platform directory iteration then this will Add a new entry rather than override the existing one.  This causes local file deletes
+				// and longer loads as we will never see the benefits of local device caching.
 				Visitor.FileTimes.Add(ServerFile, FDateTime::MinValue());
 
 				if (bDeleteFile == false)
diff --git a/Engine/Source/Runtime/NetworkFileSystem/Private/NetworkFileServerConnection.cpp b/Engine/Source/Runtime/NetworkFileSystem/Private/NetworkFileServerConnection.cpp
--- a/Engine/Source/Runtime/NetworkFileSystem/Private/NetworkFileServerConnection.cpp
+++ b/Engine/Source/Runtime/NetworkFileSystem/Private/NetworkFileServerConnection.cpp
@@ -69,7 +69,7 @@ void FNetworkFileServerClientConnection::ConvertClientFilenameToServerFilename(F
  * client path: "../../../Samples/Showcases/Elemental/Content/Elemental/Effects/FX_Snow_Cracks/Crack_02/Materials/M_SnowBlast.uasset"
  * This ensures that devicelocal-cached files will be properly timestamp checked before deletion.
  */
-static TMap<FString, FDateTime> FixupSandboxPathsForClient(FSandboxPlatformFile* Sandbox, const TMap<FString, FDateTime>& SandboxPaths, const FString& LocalEngineDir, const FString& LocalGameDir)
+static TMap<FString, FDateTime> FixupSandboxPathsForClient(FSandboxPlatformFile* Sandbox, const TMap<FString, FDateTime>& SandboxPaths, const FString& LocalEngineDir, const FString& LocalGameDir, bool bLowerCaseFiles)
 {
 	TMap<FString, FDateTime> FixedFiletimes;
 	FString SandboxEngine = Sandbox->ConvertToSandboxPath(*LocalEngineDir) + TEXT("/");
@@ -85,6 +85,11 @@ static TMap<FString, FDateTime> FixupSandboxPathsForClient(FSandboxPlatformFile*
 		FString Fixed = Sandbox->ConvertToSandboxPath(*It.Key());
 		Fixed = Fixed.Replace(*SandboxEngine, *LocalEngineDir);
 		Fixed = Fixed.Replace(*SandboxGame, *LocalGameDir);
+
+		if (bLowerCaseFiles)
+		{
+			Fixed = Fixed.ToLower();
+		}
 		FixedFiletimes.Add(Fixed, It.Value());
 	}
 	return FixedFiletimes;
@@ -644,6 +649,8 @@ bool FNetworkFileServerClientConnection::ProcessGetFileList( FArchive& In, FArch
 
 	ConnectedPlatformName = TEXT("");
 
+	bool bSendLowerCase = false;
+
 	// if we didn't find one (and this is a dumb server - no active platforms), then just use what was sent
 	if (ActiveTargetPlatforms.Num() == 0)
 	{
@@ -663,6 +670,7 @@ bool FNetworkFileServerClientConnection::ProcessGetFileList( FArchive& In, FArch
 				UE_LOG(LogFileServer, Display, TEXT("   Checking against: %s"), *ActiveTargetPlatforms[ActiveTPIndex]->PlatformName());
 				if (ActiveTargetPlatforms[ActiveTPIndex]->PlatformName() == TargetPlatformNames[TPIndex])
 				{
+					bSendLowerCase = ActiveTargetPlatforms[ActiveTPIndex]->SendLowerCaseFilePaths();
 					ConnectedPlatformName = ActiveTargetPlatforms[ActiveTPIndex]->PlatformName();
 					break;
 				}
@@ -800,7 +808,7 @@ bool FNetworkFileServerClientConnection::ProcessGetFileList( FArchive& In, FArch
 	Out << LocalGameDir;
 
 	// return the files and their timestamps
-	TMap<FString, FDateTime> FixedTimes = FixupSandboxPathsForClient(Sandbox, Visitor.FileTimes, LocalEngineDir, LocalGameDir);
+	TMap<FString, FDateTime> FixedTimes = FixupSandboxPathsForClient(Sandbox, Visitor.FileTimes, LocalEngineDir, LocalGameDir, bSendLowerCase);
 	Out << FixedTimes;
 
 	// Do it again, preventing access to non-cooked files
@@ -827,7 +835,7 @@ bool FNetworkFileServerClientConnection::ProcessGetFileList( FArchive& In, FArch
 		}
 	
 		// return the cached files and their timestamps
-		FixedTimes = FixupSandboxPathsForClient(Sandbox, VisitorForCacheDates.FileTimes, LocalEngineDir, LocalGameDir);
+		FixedTimes = FixupSandboxPathsForClient(Sandbox, VisitorForCacheDates.FileTimes, LocalEngineDir, LocalGameDir, bSendLowerCase);
 		Out << FixedTimes;
 	}
 	return true;
