changed error to warning since it isn't actually an error to send an empty file
UE-7691
#cotf

diff --git a/Engine/Source/Runtime/NetworkFileSystem/Private/NetworkFileServerConnection.cpp b/Engine/Source/Runtime/NetworkFileSystem/Private/NetworkFileServerConnection.cpp
--- a/Engine/Source/Runtime/NetworkFileSystem/Private/NetworkFileServerConnection.cpp
+++ b/Engine/Source/Runtime/NetworkFileSystem/Private/NetworkFileServerConnection.cpp
@@ -880,7 +880,7 @@ void FNetworkFileServerClientConnection::PackageFile( FString& Filename, FArchiv
 	{
 		if (!File->Size())
 		{
-			UE_LOG(LogFileServer, Error, TEXT("Sending empty file %s...."), *Filename);
+			UE_LOG(LogFileServer, Warning, TEXT("Sending empty file %s...."), *Filename);
 		}
 		else
 		{
