Added support for multiple client connections to the networkfileserver.

diff --git a/Engine/Source/Runtime/NetworkFileSystem/Private/NetworkFileServer.cpp b/Engine/Source/Runtime/NetworkFileSystem/Private/NetworkFileServer.cpp
--- a/Engine/Source/Runtime/NetworkFileSystem/Private/NetworkFileServer.cpp
+++ b/Engine/Source/Runtime/NetworkFileSystem/Private/NetworkFileServer.cpp
@@ -80,11 +80,16 @@ public:
 		return  (Running.GetValue() != 0); 
 	}
 
-	void GetAddress( FInternetAddr &Addr )
+	void GetAddress( FInternetAddr& Addr )
 	{
 		Socket->GetAddress(Addr);
 	}
 
+	void GetPeerAddress( FInternetAddr& Addr )
+	{
+		Socket->GetPeerAddress(Addr);
+	}
+
 	~FNetworkFileServerClientConnectionThreaded()
 	{
 		WorkerThread->Kill(true);
@@ -225,18 +230,23 @@ uint32 FNetworkFileServer::Run( )
 
 			if (ClientSocket != NULL)
 			{
-				TSharedPtr<FInternetAddr> Addr = ISocketSubsystem::Get(PLATFORM_SOCKETSUBSYSTEM)->CreateInternetAddr();;
+				TSharedPtr<FInternetAddr> Addr = ISocketSubsystem::Get(PLATFORM_SOCKETSUBSYSTEM)->CreateInternetAddr();
 				ClientSocket->GetAddress(*Addr);
+				TSharedPtr<FInternetAddr> PeerAddr = ISocketSubsystem::Get(PLATFORM_SOCKETSUBSYSTEM)->CreateInternetAddr();
+				ClientSocket->GetPeerAddress(*PeerAddr);
 
 				for ( auto PreviousConnection : Connections )
 				{
 					TSharedPtr<FInternetAddr> PreviousAddr = ISocketSubsystem::Get(PLATFORM_SOCKETSUBSYSTEM)->CreateInternetAddr();;
 					PreviousConnection->GetAddress( *PreviousAddr );
-					if ( *Addr == *PreviousAddr)
+					TSharedPtr<FInternetAddr> PreviousPeerAddr = ISocketSubsystem::Get(PLATFORM_SOCKETSUBSYSTEM)->CreateInternetAddr();;
+					PreviousConnection->GetPeerAddress( *PreviousPeerAddr );
+					if ( ( *Addr == *PreviousAddr ) &&
+						(*PeerAddr == *PreviousPeerAddr ) )
 					{
 						// kill hte connection 
 						PreviousConnection->Stop();
-						UE_LOG(LogFileServer, Display, TEXT( "Killing client connection %s because new client connected from same address." ), *PreviousConnection->GetDescription() );
+						UE_LOG(LogFileServer, Warning, TEXT( "Killing client connection %s because new client connected from same address." ), *PreviousConnection->GetDescription() );
 					}
 				}
 
diff --git a/Engine/Source/Runtime/Online/OnlineSubsystemSteam/Private/SocketsSteam.cpp b/Engine/Source/Runtime/Online/OnlineSubsystemSteam/Private/SocketsSteam.cpp
--- a/Engine/Source/Runtime/Online/OnlineSubsystemSteam/Private/SocketsSteam.cpp
+++ b/Engine/Source/Runtime/Online/OnlineSubsystemSteam/Private/SocketsSteam.cpp
@@ -228,6 +228,13 @@ void FSocketSteam::GetAddress(FInternetAddr& OutAddr)
 	SteamAddr.SteamChannel = SteamChannel;
 }
 
+void FSocketSteam::GetPeerAddress(FInternetAddr& OutAddr)
+{
+	// don't support this
+	FInternetAddrSteam& SteamAddr = (FInternetAddrSteam&)OutAddr;
+	new( &SteamAddr ) FInternetAddrSteam();
+}
+
 /**
  * Sets this socket into non-blocking mode
  *
diff --git a/Engine/Source/Runtime/Online/OnlineSubsystemSteam/Private/SocketsSteam.h b/Engine/Source/Runtime/Online/OnlineSubsystemSteam/Private/SocketsSteam.h
--- a/Engine/Source/Runtime/Online/OnlineSubsystemSteam/Private/SocketsSteam.h
+++ b/Engine/Source/Runtime/Online/OnlineSubsystemSteam/Private/SocketsSteam.h
@@ -189,6 +189,13 @@ public:
 	 */
 	virtual void GetAddress(FInternetAddr& OutAddr) override;
 
+	/**
+	 * Reads the address of the peer the socket is connected to and returns it
+	 * 
+	 * @param OutAddr address of the peer the socket is connected to
+	 */
+	virtual void GetPeerAddress(FInternetAddr& OutAddr) override;
+
 	/**
 	 * Sets this socket into non-blocking mode
 	 *
diff --git a/Engine/Source/Runtime/Sockets/Private/BSDIPv6Sockets/SocketsBSDIPv6.cpp b/Engine/Source/Runtime/Sockets/Private/BSDIPv6Sockets/SocketsBSDIPv6.cpp
--- a/Engine/Source/Runtime/Sockets/Private/BSDIPv6Sockets/SocketsBSDIPv6.cpp
+++ b/Engine/Source/Runtime/Sockets/Private/BSDIPv6Sockets/SocketsBSDIPv6.cpp
@@ -270,6 +270,21 @@ void FSocketBSDIPv6::GetAddress(FInternetAddr& OutAddr)
 	}
 }
 
+void FSocketBSDIPv6::GetPeerAddress(FInternetAddr& OutAddr)
+{
+	FInternetAddrBSDIPv6& Addr = (FInternetAddrBSDIPv6&)OutAddr;
+	SOCKLEN Size = sizeof(sockaddr_in6);
+
+	// Figure out what ip/port we are bound to
+	bool bOk = getpeername(Socket, Addr, &Size) == 0;
+
+	if (bOk == false)
+	{
+		check(SocketSubsystem);
+		UE_LOG(LogSockets, Error, TEXT("Failed to read address for socket (%s)"), SocketSubsystem->GetSocketError());
+	}
+}
+
 
 bool FSocketBSDIPv6::SetNonBlocking(bool bIsNonBlocking)
 {
diff --git a/Engine/Source/Runtime/Sockets/Private/BSDIPv6Sockets/SocketsBSDIPv6.h b/Engine/Source/Runtime/Sockets/Private/BSDIPv6Sockets/SocketsBSDIPv6.h
--- a/Engine/Source/Runtime/Sockets/Private/BSDIPv6Sockets/SocketsBSDIPv6.h
+++ b/Engine/Source/Runtime/Sockets/Private/BSDIPv6Sockets/SocketsBSDIPv6.h
@@ -90,6 +90,8 @@ public:
 
 	virtual void GetAddress(FInternetAddr& OutAddr) override;
 
+	virtual void GetPeerAddress(FInternetAddr& OutAddr) override;
+
 	virtual bool SetNonBlocking(bool bIsNonBlocking = true) override;
 
 	virtual bool SetBroadcast(bool bAllowBroadcast = true) override;
diff --git a/Engine/Source/Runtime/Sockets/Private/BSDSockets/SocketsBSD.cpp b/Engine/Source/Runtime/Sockets/Private/BSDSockets/SocketsBSD.cpp
--- a/Engine/Source/Runtime/Sockets/Private/BSDSockets/SocketsBSD.cpp
+++ b/Engine/Source/Runtime/Sockets/Private/BSDSockets/SocketsBSD.cpp
@@ -265,6 +265,21 @@ void FSocketBSD::GetAddress(FInternetAddr& OutAddr)
 }
 
 
+void FSocketBSD::GetPeerAddress(FInternetAddr& OutAddr)
+{
+	FInternetAddrBSD& Addr = (FInternetAddrBSD&)OutAddr;
+	SOCKLEN Size = sizeof(sockaddr_in);
+
+	// Figure out what ip/port we are bound to
+	bool bOk = getpeername(Socket, Addr, &Size) == 0;
+
+	if (bOk == false)
+	{
+		check(SocketSubsystem);
+		UE_LOG(LogSockets, Error, TEXT("Failed to read address for socket (%s)"), SocketSubsystem->GetSocketError());
+	}
+}
+
 bool FSocketBSD::SetNonBlocking(bool bIsNonBlocking)
 {
 	u_long Value = bIsNonBlocking ? true : false;
diff --git a/Engine/Source/Runtime/Sockets/Private/BSDSockets/SocketsBSD.h b/Engine/Source/Runtime/Sockets/Private/BSDSockets/SocketsBSD.h
--- a/Engine/Source/Runtime/Sockets/Private/BSDSockets/SocketsBSD.h
+++ b/Engine/Source/Runtime/Sockets/Private/BSDSockets/SocketsBSD.h
@@ -94,6 +94,7 @@ public:
 	virtual bool Wait(ESocketWaitConditions::Type Condition, FTimespan WaitTime) override;
 	virtual ESocketConnectionState GetConnectionState() override;
 	virtual void GetAddress(FInternetAddr& OutAddr) override;
+	virtual void GetPeerAddress(FInternetAddr& OutAddr) override;
 	virtual bool SetNonBlocking(bool bIsNonBlocking = true) override;
 	virtual bool SetBroadcast(bool bAllowBroadcast = true) override;
 	virtual bool JoinMulticastGroup(const FInternetAddr& GroupAddress) override;
diff --git a/Engine/Source/Runtime/Sockets/Public/Sockets.h b/Engine/Source/Runtime/Sockets/Public/Sockets.h
--- a/Engine/Source/Runtime/Sockets/Public/Sockets.h
+++ b/Engine/Source/Runtime/Sockets/Public/Sockets.h
@@ -177,6 +177,13 @@ public:
 	 */
 	virtual void GetAddress(FInternetAddr& OutAddr) = 0;
 
+	/**
+	 * Reads the address of the peer the socket is connected to
+	 *
+	 * @param OutAddr address of the peer the socket is connected to
+	 */
+	virtual void GetPeerAddress(FInternetAddr& OutAddr) = 0;
+
 	/**
 	 * Sets this socket into non-blocking mode
 	 *
