Fixed bug with directories for cook on the fly clients not having cached files cleaned up.  The server needs to convert the directories to client space as clients can have different root paths.

diff --git a/Engine/Source/Runtime/NetworkFileSystem/Private/NetworkFileServerConnection.cpp b/Engine/Source/Runtime/NetworkFileSystem/Private/NetworkFileServerConnection.cpp
--- a/Engine/Source/Runtime/NetworkFileSystem/Private/NetworkFileServerConnection.cpp
+++ b/Engine/Source/Runtime/NetworkFileSystem/Private/NetworkFileServerConnection.cpp
@@ -829,7 +829,12 @@ bool FNetworkFileServerClientConnection::ProcessGetFileList( FArchive& In, FArch
 		for (const auto& RootPath : RootContentPaths)
 		{
 			const FString& ContentFolder = FPackageName::LongPackageNameToFilename(RootPath);
-			ContentFolders.Add(ContentFolder);
+
+			FString ConnectedContentFolder = ContentFolder;
+			ConnectedContentFolder.ReplaceInline(*LocalEngineDir, *ConnectedEngineDir);
+			ConnectedContentFolder.ReplaceInline(*LocalGameDir, *ConnectedGameDir);
+
+			ContentFolders.Add(ConnectedContentFolder);
 		}
 		Out << ContentFolders;
 
