Don't add a second slash if we already have one when fixing up sandbox paths for client.
Fixes crash in cook on the fly.

diff --git a/Engine/Source/Runtime/NetworkFileSystem/Private/NetworkFileServerConnection.cpp b/Engine/Source/Runtime/NetworkFileSystem/Private/NetworkFileServerConnection.cpp
--- a/Engine/Source/Runtime/NetworkFileSystem/Private/NetworkFileServerConnection.cpp
+++ b/Engine/Source/Runtime/NetworkFileSystem/Private/NetworkFileServerConnection.cpp
@@ -72,12 +72,16 @@ void FNetworkFileServerClientConnection::ConvertClientFilenameToServerFilename(F
 static TMap<FString, FDateTime> FixupSandboxPathsForClient(FSandboxPlatformFile* Sandbox, const TMap<FString, FDateTime>& SandboxPaths, const FString& LocalEngineDir, const FString& LocalGameDir, bool bLowerCaseFiles)
 {
 	TMap<FString, FDateTime> FixedFiletimes;
-	FString SandboxEngine = Sandbox->ConvertToSandboxPath(*LocalEngineDir) + TEXT("/");
+	FString SandboxEngine = Sandbox->ConvertToSandboxPath(*LocalEngineDir);
+	if (SandboxEngine.EndsWith(TEXT("/")) == false)
+	{
+		SandboxEngine += TEXT("/");
+	}
 
 	// we need to add an extra bit to the game path to make the sandbox convert it correctly (investigate?)
 	// @todo: double check this
 	FString SandboxGame = Sandbox->ConvertToSandboxPath(*(LocalGameDir + TEXT("a.txt"))).Replace(TEXT("a.txt"), TEXT(""));
-
+	
 	// since the sandbox remaps from A/B/C to C, and the client has no idea of this, we need to put the files
 	// into terms of the actual LocalGameDir, which is all that the client knows about
 	for (TMap<FString, FDateTime>::TConstIterator It(SandboxPaths); It; ++It)
@@ -633,7 +637,7 @@ void FNetworkFileServerClientConnection::ProcessToAbsolutePathForWrite( FArchive
 bool FNetworkFileServerClientConnection::ProcessGetFileList( FArchive& In, FArchive& Out )
 {
 	// get the list of directories to process
-	TArray<FString> TargetPlatformNames;
+ 	TArray<FString> TargetPlatformNames;
 	FString GameName;
 	FString EngineRelativePath;
 	FString GameRelativePath;
