#html5 serve compressed gz correct when cook on the fly.

diff --git a/Engine/Source/Runtime/NetworkFileSystem/Private/NetworkFileServerHttp.cpp b/Engine/Source/Runtime/NetworkFileSystem/Private/NetworkFileServerHttp.cpp
--- a/Engine/Source/Runtime/NetworkFileSystem/Private/NetworkFileServerHttp.cpp
+++ b/Engine/Source/Runtime/NetworkFileSystem/Private/NetworkFileServerHttp.cpp
@@ -368,13 +368,28 @@ int FNetworkFileServerHttp::CallBack_HTTP(
 
 				// file up the header. 
 				TCHAR Header[1024];
-				int Length  = FCString::Sprintf(Header,
-					TEXT("HTTP/1.1 200 OK\x0d\x0a")
-					TEXT("Server: Unreal File Server\x0d\x0a")
-					TEXT("Connection: close\x0d\x0a")
-					TEXT("Content-Type: %s \x0d\x0a")
-					TEXT("Content-Length: %u\x0d\x0a\x0d\x0a"),
-					Mime,FileData.Num());
+				int Length = 0; 
+				if (FilePath.Contains("gz"))
+				{
+					Length = FCString::Sprintf(Header,
+						TEXT("HTTP/1.1 200 OK\x0d\x0a")
+						TEXT("Server: Unreal File Server\x0d\x0a")
+						TEXT("Connection: close\x0d\x0a")
+						TEXT("Content-Type: %s \x0d\x0a")
+						TEXT("Content-Encoding: gzip\x0d\x0a")
+						TEXT("Content-Length: %u\x0d\x0a\x0d\x0a"),
+						Mime, FileData.Num());
+				}
+				else
+				{
+					Length = FCString::Sprintf(Header,
+						TEXT("HTTP/1.1 200 OK\x0d\x0a")
+						TEXT("Server: Unreal File Server\x0d\x0a")
+						TEXT("Connection: close\x0d\x0a")
+						TEXT("Content-Type: %s \x0d\x0a")
+						TEXT("Content-Length: %u\x0d\x0a\x0d\x0a"),
+						Mime, FileData.Num());
+				}
 
 				// make space for the whole file in our out buffer. 
 				BufferInfo->Out.Append((uint8*)TCHAR_TO_ANSI(Header),Length); 
