Switched a search for punctuation to be case sensitive for performance reasons.

diff --git a/Engine/Source/Runtime/NetworkFileSystem/Private/NetworkFileServerConnection.cpp b/Engine/Source/Runtime/NetworkFileSystem/Private/NetworkFileServerConnection.cpp
--- a/Engine/Source/Runtime/NetworkFileSystem/Private/NetworkFileServerConnection.cpp
+++ b/Engine/Source/Runtime/NetworkFileSystem/Private/NetworkFileServerConnection.cpp
@@ -73,7 +73,7 @@ static TMap<FString, FDateTime> FixupSandboxPathsForClient(FSandboxPlatformFile*
 {
 	TMap<FString, FDateTime> FixedFiletimes;
 	FString SandboxEngine = Sandbox->ConvertToSandboxPath(*LocalEngineDir);
-	if (SandboxEngine.EndsWith(TEXT("/")) == false)
+	if (SandboxEngine.EndsWith(TEXT("/"), ESearchCase::CaseSensitive) == false)
 	{
 		SandboxEngine += TEXT("/");
 	}
