Copying //UE4/Orion-Staging to //UE4/Main (Source: //Orion/Dev-General @ 3483207)

diff --git a/Engine/Source/Runtime/NetworkFileSystem/Private/NetworkFileServerConnection.cpp b/Engine/Source/Runtime/NetworkFileSystem/Private/NetworkFileServerConnection.cpp
--- a/Engine/Source/Runtime/NetworkFileSystem/Private/NetworkFileServerConnection.cpp
+++ b/Engine/Source/Runtime/NetworkFileSystem/Private/NetworkFileServerConnection.cpp
@@ -36,34 +36,12 @@ void GetSandboxRootDirectories(FSandboxPlatformFile* Sandbox, FString& SandboxEn
 /* FNetworkFileServerClientConnection structors
  *****************************************************************************/
 
-FNetworkFileServerClientConnection::FNetworkFileServerClientConnection( const FFileRequestDelegate& InFileRequestDelegate, 
-		const FRecompileShadersDelegate& InRecompileShadersDelegate, const FSandboxPathDelegate& InSandboxPathOverrideDelegate, FOnFileModifiedDelegate* InOnFileModifiedCallback, const TArray<ITargetPlatform*>& InActiveTargetPlatforms )
+FNetworkFileServerClientConnection::FNetworkFileServerClientConnection( const FNetworkFileDelegateContainer* InNetworkFileDelegates, const TArray<ITargetPlatform*>& InActiveTargetPlatforms )
 	: LastHandleId(0)
 	, Sandbox(NULL)
+	, NetworkFileDelegates(InNetworkFileDelegates)
 	, ActiveTargetPlatforms(InActiveTargetPlatforms)
-{
-
-	OnFileModifiedCallback = InOnFileModifiedCallback;
-	if (OnFileModifiedCallback)
-	{
-		OnFileModifiedCallback->AddRaw(this, &FNetworkFileServerClientConnection::FileModifiedCallback);
-	}
-
-	if (InFileRequestDelegate.IsBound())
-	{
-		FileRequestDelegate = InFileRequestDelegate;
-	}
-
-	if (InRecompileShadersDelegate.IsBound())
-	{
-		RecompileShadersDelegate = InRecompileShadersDelegate;
-	}
-
-	if ( InSandboxPathOverrideDelegate.IsBound() )
-	{
-		SandboxPathOverrideDelegate = InSandboxPathOverrideDelegate;
-	}
-	
+{	
 	//stats
 	FileRequestDelegateTime = 0.0;
 	PackageFileTime = 0.0;
@@ -74,12 +52,20 @@ FNetworkFileServerClientConnection::FNetworkFileServerClientConnection( const FF
 	PackageRequestsSucceeded = 0;
 	PackageRequestsFailed = 0;
 	FileBytesSent = 0;
+
+	if ( NetworkFileDelegates && NetworkFileDelegates->OnFileModifiedCallback )
+	{
+		NetworkFileDelegates->OnFileModifiedCallback->AddRaw(this, &FNetworkFileServerClientConnection::FileModifiedCallback);
+	}
 }
 
 
 FNetworkFileServerClientConnection::~FNetworkFileServerClientConnection( )
 {
-	OnFileModifiedCallback->RemoveAll(this);
+	if (NetworkFileDelegates && NetworkFileDelegates->OnFileModifiedCallback)
+	{
+		NetworkFileDelegates->OnFileModifiedCallback->RemoveAll(this);
+	}
 
 	// close all the files the client had opened through us when the client disconnects
 	for (TMap<uint64, IFileHandle*>::TIterator It(OpenFiles); It; ++It)
@@ -395,7 +381,7 @@ void FNetworkFileServerClientConnection::ProcessOpenFile( FArchive& In, FArchive
 	}
 
 	TArray<FString> NewUnsolictedFiles;
-	FileRequestDelegate.ExecuteIfBound(Filename, ConnectedPlatformName, NewUnsolictedFiles);
+	NetworkFileDelegates->FileRequestDelegate.ExecuteIfBound(Filename, ConnectedPlatformName, NewUnsolictedFiles);
 
 	FDateTime ServerTimeStamp = Sandbox->GetTimeStamp(*Filename);
 	int64 ServerFileSize = 0;
@@ -542,7 +528,7 @@ void FNetworkFileServerClientConnection::ProcessGetFileInfo( FArchive& In, FArch
 	if (Info.FileExists)
 	{
 		TArray<FString> NewUnsolictedFiles;
-		FileRequestDelegate.ExecuteIfBound(Filename, ConnectedPlatformName, NewUnsolictedFiles);
+		NetworkFileDelegates->FileRequestDelegate.ExecuteIfBound(Filename, ConnectedPlatformName, NewUnsolictedFiles);
 	}
 
 	// get the rest of the info
@@ -728,14 +714,34 @@ bool FNetworkFileServerClientConnection::ProcessGetFileList( FArchive& In, FArch
 	FString EngineRelativePath;
 	FString GameRelativePath;
 	TArray<FString> RootDirectories;
-	bool bIsStreamingRequest = false;
+	
+	EConnectionFlags ConnectionFlags;
+
+	FString ClientVersionInfo;
 
 	In << TargetPlatformNames;
 	In << GameName;
 	In << EngineRelativePath;
 	In << GameRelativePath;
 	In << RootDirectories;
-	In << bIsStreamingRequest;
+	In << ConnectionFlags;
+	In << ClientVersionInfo;
+
+	if ( NetworkFileDelegates->NewConnectionDelegate.IsBound() )
+	{
+		bool bIsValidVersion = true;
+		for ( const FString& TargetPlatform : TargetPlatformNames )
+		{
+			bIsValidVersion &= NetworkFileDelegates->NewConnectionDelegate.Execute(ClientVersionInfo, TargetPlatform );
+		}
+		if ( bIsValidVersion == false )
+		{
+			return false;
+		}
+	}
+
+	const bool bIsStreamingRequest = (ConnectionFlags & EConnectionFlags::Streaming) == EConnectionFlags::Streaming;
+	const bool bIsPrecookedIterativeRequest = (ConnectionFlags & EConnectionFlags::PreCookedIterative) == EConnectionFlags::PreCookedIterative;
 
 	ConnectedPlatformName = TEXT("");
 
@@ -810,9 +816,9 @@ bool FNetworkFileServerClientConnection::ProcessGetFileList( FArchive& In, FArch
 	// figure out the sandbox directory
 	// @todo: This should use FPlatformMisc::SavedDirectory(GameName)
 	FString SandboxDirectory;
-	if ( SandboxPathOverrideDelegate.IsBound() )
+	if (NetworkFileDelegates->SandboxPathOverrideDelegate.IsBound() )
 	{
-		SandboxDirectory = SandboxPathOverrideDelegate.Execute();
+		SandboxDirectory = NetworkFileDelegates->SandboxPathOverrideDelegate.Execute();
 		// if the sandbox directory delegate returns a path with the platform name in it then replace it :)
 		SandboxDirectory.ReplaceInline(TEXT("[Platform]"), *ConnectedPlatformName);
 	}
@@ -864,7 +870,7 @@ bool FNetworkFileServerClientConnection::ProcessGetFileList( FArchive& In, FArch
 	RecompileData.ShaderPlatform = -1;
 	RecompileData.ModifiedFiles = NULL;
 	RecompileData.MeshMaterialMaps = NULL;
-	RecompileShadersDelegate.ExecuteIfBound(RecompileData);
+	NetworkFileDelegates->RecompileShadersDelegate.ExecuteIfBound(RecompileData);
 
 	UE_LOG(LogFileServer, Display, TEXT("Getting files for %d directories, game = %s, platform = %s"), RootDirectories.Num(), *GameName, *ConnectedPlatformName);
 	UE_LOG(LogFileServer, Display, TEXT("    Sandbox dir = %s"), *SandboxDirectory);
@@ -894,6 +900,7 @@ bool FNetworkFileServerClientConnection::ProcessGetFileList( FArchive& In, FArch
 		DirectoriesToSkip.Add(FString(RootDirectories[DirIndex] / TEXT("Saved/Logs")));
 		DirectoriesToSkip.Add(FString(RootDirectories[DirIndex] / TEXT("Saved/Sandboxes")));
 		DirectoriesToSkip.Add(FString(RootDirectories[DirIndex] / TEXT("Saved/Cooked")));
+		DirectoriesToSkip.Add(FString(RootDirectories[DirIndex] / TEXT("Saved/EditorCooked")));
 		DirectoriesToSkip.Add(FString(RootDirectories[DirIndex] / TEXT("Saved/ShaderDebugInfo")));
 		DirectoriesToSkip.Add(FString(RootDirectories[DirIndex] / TEXT("Saved/StagedBuilds")));
 		DirectoriesToSkip.Add(FString(RootDirectories[DirIndex] / TEXT("Intermediate")));
@@ -1000,6 +1007,18 @@ bool FNetworkFileServerClientConnection::ProcessGetFileList( FArchive& In, FArch
 		FixedTimes = FixupSandboxPathsForClient(VisitorForCacheDates.FileTimes);
 		Out << FixedTimes;
 	}
+
+
+
+	if ( bIsPrecookedIterativeRequest )
+	{
+		TMap<FString, FDateTime> PrecookedList;
+		NetworkFileDelegates->InitialPrecookedListDelegate.ExecuteIfBound(ConnectedPlatformName, PrecookedList);
+
+		FixedTimes = FixupSandboxPathsForClient(PrecookedList);
+		Out << FixedTimes;
+	}
+
 	return true;
 }
 
@@ -1010,15 +1029,17 @@ void FNetworkFileServerClientConnection::FileModifiedCallback( const FString& Fi
 	// do we care about this file???
 
 	// translation here?
-	ModifiedFiles.Add(Filename);
+	ModifiedFiles.AddUnique(Filename);
 }
 
 void FNetworkFileServerClientConnection::ProcessHeartbeat( FArchive& In, FArchive& Out )
 {
 	TArray<FString> FixedupModifiedFiles;
 	// Protect the array
+	if (Sandbox)
 	{
 		FScopeLock Lock(&ModifiedFilesSection);
+		
 		for (const auto& ModifiedFile : ModifiedFiles)
 		{
 			FixedupModifiedFiles.Add(FixupSandboxPathForClient(ModifiedFile));
@@ -1104,7 +1125,7 @@ void FNetworkFileServerClientConnection::ProcessRecompileShaders( FArchive& In,
 	In << RecompileData.SerializedShaderResources;
 	In << RecompileData.bCompileChangedShaders;
 
-	RecompileShadersDelegate.ExecuteIfBound(RecompileData);
+	NetworkFileDelegates->RecompileShadersDelegate.ExecuteIfBound(RecompileData);
 
 	// tell other side what to do!
 	Out << RecompileModifiedFiles;
@@ -1131,7 +1152,7 @@ void FNetworkFileServerClientConnection::ProcessSyncFile( FArchive& In, FArchive
 
 	TArray<FString> NewUnsolictedFiles;
 
-	FileRequestDelegate.ExecuteIfBound(Filename, ConnectedPlatformName, NewUnsolictedFiles);
+	NetworkFileDelegates->FileRequestDelegate.ExecuteIfBound(Filename, ConnectedPlatformName, NewUnsolictedFiles);
 
 	FileRequestDelegateTime += 1000.0f * float(FPlatformTime::Seconds() - StartTime);
 	StartTime = FPlatformTime::Seconds();
diff --git a/Engine/Source/Runtime/NetworkFileSystem/Private/NetworkFileServerConnection.h b/Engine/Source/Runtime/NetworkFileSystem/Private/NetworkFileServerConnection.h
--- a/Engine/Source/Runtime/NetworkFileSystem/Private/NetworkFileServerConnection.h
+++ b/Engine/Source/Runtime/NetworkFileSystem/Private/NetworkFileServerConnection.h
@@ -20,10 +20,9 @@ public:
 	 * Creates and initializes a new instance.
 	 *
 	 * @param InSocket - The client socket to use.
-	 * @param InFileRequestDelegate - A delegate to be invoked when the client requests a file.
+	 * @param NetworkFileDelegates- delegates the client calls when events from the client happen
 	 */
-	FNetworkFileServerClientConnection(const FFileRequestDelegate& InFileRequestDelegate, 
-		const FRecompileShadersDelegate& InRecompileShadersDelegate, const FSandboxPathDelegate& InSandboxPathOverrideDelegate, FOnFileModifiedDelegate* OnFileModifiedDelegate, const TArray<ITargetPlatform*>& InActiveTargetPlatforms );
+	FNetworkFileServerClientConnection( const FNetworkFileDelegateContainer* NetworkFileDelegates , const TArray<ITargetPlatform*>& InActiveTargetPlatforms );
 
 	/**
 	 * Destructor.
@@ -227,15 +226,7 @@ private:
 	// Holds the list of directories being watched.
 	TArray<FString> WatchedDirectories;
 
-	// Holds a delegate to be invoked on every sync request.
-	FFileRequestDelegate FileRequestDelegate;
-
-	// Holds a delegate to be invoked when a client requests a shader recompile.
-	FRecompileShadersDelegate RecompileShadersDelegate;
-
-	FSandboxPathDelegate SandboxPathOverrideDelegate;
-
-	FOnFileModifiedDelegate* OnFileModifiedCallback;
+	const FNetworkFileDelegateContainer* NetworkFileDelegates;
 
 	// cached copy of the active target platforms (if any)
 	const TArray<ITargetPlatform*>& ActiveTargetPlatforms;
diff --git a/Engine/Source/Runtime/NetworkFileSystem/Private/NetworkFileServerHttp.cpp b/Engine/Source/Runtime/NetworkFileSystem/Private/NetworkFileServerHttp.cpp
--- a/Engine/Source/Runtime/NetworkFileSystem/Private/NetworkFileServerHttp.cpp
+++ b/Engine/Source/Runtime/NetworkFileSystem/Private/NetworkFileServerHttp.cpp
@@ -17,9 +17,8 @@ class FNetworkFileServerClientConnectionHTTP : public FNetworkFileServerClientCo
 {
 
 public:
-	FNetworkFileServerClientConnectionHTTP(const FFileRequestDelegate& InFileRequestDelegate,
-		const FRecompileShadersDelegate& InRecompileShadersDelegate, const FSandboxPathDelegate& SandboxFileDelegate, FOnFileModifiedDelegate* OnFileModifiedCallback, const TArray<ITargetPlatform*>& InActiveTargetPlatforms )
-		:  FNetworkFileServerClientConnection( InFileRequestDelegate,InRecompileShadersDelegate, SandboxFileDelegate, OnFileModifiedCallback,InActiveTargetPlatforms)
+	FNetworkFileServerClientConnectionHTTP(const FNetworkFileDelegateContainer* NetworkFileDelegates, const TArray<ITargetPlatform*>& InActiveTargetPlatforms )
+		:  FNetworkFileServerClientConnection( NetworkFileDelegates, InActiveTargetPlatforms)
 	{
 	}
 
@@ -81,10 +80,7 @@ static struct lws_protocols Protocols[] = {
 
 FNetworkFileServerHttp::FNetworkFileServerHttp(
 	int32 InPort,
-	const FFileRequestDelegate* InFileRequestDelegate,
-	const FRecompileShadersDelegate* InRecompileShadersDelegate,
-	const FSandboxPathDelegate* InSandboxPathDelegate,
-	FOnFileModifiedDelegate* InOnFileModifiedCallback,
+	FNetworkFileDelegateContainer InNetworkFileDelegateContainer,
 	const TArray<ITargetPlatform*>& InActiveTargetPlatforms
 	)
 	:ActiveTargetPlatforms(InActiveTargetPlatforms)
@@ -97,22 +93,7 @@ FNetworkFileServerHttp::FNetworkFileServerHttp(
 
 	UE_LOG(LogFileServer, Warning, TEXT("Unreal Network Http File Server starting up..."));
 
-	if (InFileRequestDelegate && InFileRequestDelegate->IsBound())
-	{
-		FileRequestDelegate = *InFileRequestDelegate;
-	}
-
-	if (InRecompileShadersDelegate && InRecompileShadersDelegate->IsBound())
-	{
-		RecompileShadersDelegate = *InRecompileShadersDelegate;
-	}
-
-	if ( InSandboxPathDelegate && InSandboxPathDelegate->IsBound())
-	{
-		SandboxPathDelegate = *InSandboxPathDelegate;
-	}
-	OnFileModifiedCallback = InOnFileModifiedCallback;
-
+	NetworkFileDelegates = InNetworkFileDelegateContainer;
 
 	StopRequested.Reset();
 	Ready.Reset();
@@ -258,7 +239,7 @@ FNetworkFileServerHttp::~FNetworkFileServerHttp()
 
 FNetworkFileServerClientConnectionHTTP* FNetworkFileServerHttp::CreateNewConnection()
 {
-	return new FNetworkFileServerClientConnectionHTTP(FileRequestDelegate,RecompileShadersDelegate,SandboxPathDelegate, OnFileModifiedCallback,ActiveTargetPlatforms);
+	return new FNetworkFileServerClientConnectionHTTP(&NetworkFileDelegates,ActiveTargetPlatforms);
 }
 
 // Have a similar process function for the normal tcp connection.
diff --git a/Engine/Source/Runtime/NetworkFileSystem/Private/NetworkFileServerHttp.h b/Engine/Source/Runtime/NetworkFileSystem/Private/NetworkFileServerHttp.h
--- a/Engine/Source/Runtime/NetworkFileSystem/Private/NetworkFileServerHttp.h
+++ b/Engine/Source/Runtime/NetworkFileSystem/Private/NetworkFileServerHttp.h
@@ -48,9 +48,7 @@ class FNetworkFileServerHttp
 {
 
 public:
-	FNetworkFileServerHttp(int32 InPort, const FFileRequestDelegate* InFileRequestDelegate,
-		const FRecompileShadersDelegate* InRecompileShadersDelegate, const FSandboxPathDelegate* SandboxPathOverrideDelegate, 
-		FOnFileModifiedDelegate* OnFileModifiedDelegate, const TArray<ITargetPlatform*>& InActiveTargetPlatforms );
+	FNetworkFileServerHttp(int32 InPort, FNetworkFileDelegateContainer InNetworkFileDelegateContainer, const TArray<ITargetPlatform*>& InActiveTargetPlatforms );
 
 	// INetworkFileServer Interface.
 
@@ -82,15 +80,7 @@ private:
 	// factory method for creating a new Client Connection.
 	class FNetworkFileServerClientConnectionHTTP* CreateNewConnection();
 
-	// Holds a delegate to be invoked on every sync request.
-	FFileRequestDelegate FileRequestDelegate;
-
-	// Holds a delegate to be invoked when a client requests a shader recompile.
-	FRecompileShadersDelegate RecompileShadersDelegate;
-
-	FSandboxPathDelegate SandboxPathDelegate;
-
-	FOnFileModifiedDelegate* OnFileModifiedCallback;
+	FNetworkFileDelegateContainer NetworkFileDelegates;
 
 	// cached copy of the active target platforms (if any)
 	const TArray<ITargetPlatform*> ActiveTargetPlatforms;
diff --git a/Engine/Source/Runtime/NetworkFileSystem/Private/NetworkFileSystemModule.cpp b/Engine/Source/Runtime/NetworkFileSystem/Private/NetworkFileSystemModule.cpp
--- a/Engine/Source/Runtime/NetworkFileSystem/Private/NetworkFileSystemModule.cpp
+++ b/Engine/Source/Runtime/NetworkFileSystem/Private/NetworkFileSystemModule.cpp
@@ -23,9 +23,7 @@ public:
 
 	// INetworkFileSystemModule interface
 
-	virtual INetworkFileServer* CreateNetworkFileServer( bool bLoadTargetPlatforms, int32 Port, 
-		const FFileRequestDelegate* InFileRequestDelegate, const FRecompileShadersDelegate* InRecompileShadersDelegate, 
-		const FSandboxPathDelegate* SandboxPathOverrideDelegate, FOnFileModifiedDelegate* OnFileModifiedCallback, const ENetworkFileServerProtocol Protocol ) const override
+	virtual INetworkFileServer* CreateNetworkFileServer( bool bLoadTargetPlatforms, int32 Port, FNetworkFileDelegateContainer NetworkFileDelegateContainer, const ENetworkFileServerProtocol Protocol ) const override
 	{
 		TArray<ITargetPlatform*> ActiveTargetPlatforms;
 		if (bLoadTargetPlatforms)
@@ -48,10 +46,10 @@ public:
 		{
 #if ENABLE_HTTP_FOR_NFS
 		case NFSP_Http: 
-			return new FNetworkFileServerHttp(Port, InFileRequestDelegate, InRecompileShadersDelegate, SandboxPathOverrideDelegate, OnFileModifiedCallback, ActiveTargetPlatforms);
+			return new FNetworkFileServerHttp(Port, NetworkFileDelegateContainer, ActiveTargetPlatforms);
 #endif
 		case NFSP_Tcp:
-			return new FNetworkFileServer(Port, InFileRequestDelegate, InRecompileShadersDelegate, SandboxPathOverrideDelegate, OnFileModifiedCallback, ActiveTargetPlatforms);
+			return new FNetworkFileServer(Port, NetworkFileDelegateContainer, ActiveTargetPlatforms);
 		}
  
 		return NULL;
diff --git a/Engine/Source/Runtime/NetworkFileSystem/Public/Interfaces/INetworkFileSystemModule.h b/Engine/Source/Runtime/NetworkFileSystem/Public/Interfaces/INetworkFileSystemModule.h
--- a/Engine/Source/Runtime/NetworkFileSystem/Public/Interfaces/INetworkFileSystemModule.h
+++ b/Engine/Source/Runtime/NetworkFileSystem/Public/Interfaces/INetworkFileSystemModule.h
@@ -62,6 +62,49 @@ DECLARE_DELEGATE_RetVal( FString, FSandboxPathDelegate);
 DECLARE_MULTICAST_DELEGATE_OneParam( FOnFileModifiedDelegate, const FString& );
 
 
+/**
+ * Delegate which is called when a new connection is made to a file server client
+ * 
+ * @param 1 Version string
+ * @param 2 Platform name
+ * @return return false if the connection should be destroyed
+ */
+DECLARE_DELEGATE_RetVal_TwoParams( bool, FNewConnectionDelegate, const FString&, const FString& );
+
+
+/**
+ * Delegate which returns a list of all the files which should already be deployed to the devkit
+ *
+ * @param 1 IN, Platform to get precooked file list
+ * @param 2 OUT, list of precooked files 
+ */
+typedef TMap<FString,FDateTime> FFileTimeMap;
+DECLARE_DELEGATE_TwoParams( FInitialPrecookedListDelegate, const FString&, FFileTimeMap& );
+
+
+
+// container struct for delegates which the network file system uses
+struct FNetworkFileDelegateContainer
+{
+public:
+	FNetworkFileDelegateContainer() : 
+		NewConnectionDelegate(nullptr), 
+		InitialPrecookedListDelegate(nullptr),
+		SandboxPathOverrideDelegate(nullptr),
+		FileRequestDelegate(nullptr),
+		RecompileShadersDelegate(nullptr),
+		OnFileModifiedCallback(nullptr)
+	{}
+	FNewConnectionDelegate NewConnectionDelegate; 
+	FInitialPrecookedListDelegate InitialPrecookedListDelegate;
+	FSandboxPathDelegate SandboxPathOverrideDelegate;
+	FFileRequestDelegate FileRequestDelegate;
+	FRecompileShadersDelegate RecompileShadersDelegate;
+
+	FOnFileModifiedDelegate* OnFileModifiedCallback; // this is called from other systems to notify the network file system that a file has been modified hence the terminology callback
+};
+
+
 enum ENetworkFileServerProtocol
 {
 	NFSP_Tcp,
@@ -86,7 +129,7 @@ public:
 	 *
 	 * @return The new file server, or nullptr if creation failed.
 	 */
-	virtual INetworkFileServer* CreateNetworkFileServer( bool bLoadTargetPlatforms, int32 Port = -1, const FFileRequestDelegate* InFileRequestDelegate = nullptr, const FRecompileShadersDelegate* InRecompileShadersDelegate = nullptr, const FSandboxPathDelegate* SandboxPathOverrideDelegate = nullptr, FOnFileModifiedDelegate* OnFileModifiedDelegate = nullptr, const ENetworkFileServerProtocol Protocol = NFSP_Tcp ) const = 0;
+	virtual INetworkFileServer* CreateNetworkFileServer( bool bLoadTargetPlatforms, int32 Port = -1, FNetworkFileDelegateContainer InNetworkFileDelegateContainer = FNetworkFileDelegateContainer(), const ENetworkFileServerProtocol Protocol = NFSP_Tcp ) const = 0;
 
 public:
 
diff --git a/Engine/Source/Runtime/PakFile/Private/IPlatformFilePak.cpp b/Engine/Source/Runtime/PakFile/Private/IPlatformFilePak.cpp
--- a/Engine/Source/Runtime/PakFile/Private/IPlatformFilePak.cpp
+++ b/Engine/Source/Runtime/PakFile/Private/IPlatformFilePak.cpp
@@ -4076,7 +4076,13 @@ bool FPakPlatformFile::Initialize(IPlatformFile* Inner, const TCHAR* CmdLine)
 	//if we are using a fileserver, then dont' mount paks automatically.  We only want to read files from the server.
 	FString FileHostIP;
 	const bool bCookOnTheFly = FParse::Value(FCommandLine::Get(), TEXT("filehostip"), FileHostIP);
-	bMountPaks = !bCookOnTheFly;
+	const bool bPreCookedNetwork = FParse::Param(FCommandLine::Get(), TEXT("precookednetwork") );
+	if (bPreCookedNetwork)
+	{
+		// precooked network builds are dependent on cook on the fly
+		check(bCookOnTheFly);
+	}
+	bMountPaks &= (!bCookOnTheFly || bPreCookedNetwork);
 #endif
 
 	if (bMountPaks)
diff --git a/Engine/Source/Runtime/SandboxFile/Public/IPlatformFileSandboxWrapper.h b/Engine/Source/Runtime/SandboxFile/Public/IPlatformFileSandboxWrapper.h
--- a/Engine/Source/Runtime/SandboxFile/Public/IPlatformFileSandboxWrapper.h
+++ b/Engine/Source/Runtime/SandboxFile/Public/IPlatformFileSandboxWrapper.h
@@ -475,7 +475,7 @@ public:
 				// Now convert the sandbox path back to engine path because the sandbox folder should not be exposed
 				// to the engine and remain transparent.
 				LocalFilename = LocalFilename.Mid(SandboxFile.GetSandboxDirectory().Len());
-				if (LocalFilename.StartsWith(TEXT("Engine/")))
+				if (LocalFilename.StartsWith(TEXT("Engine/")) || (FCString::Stricmp( *LocalFilename, TEXT("Engine") ) == 0))
 				{
 					LocalFilename = SandboxFile.GetAbsoluteRootDirectory() / LocalFilename;
 				}
diff --git a/Engine/Source/Runtime/Sockets/Private/ServerTOC.cpp b/Engine/Source/Runtime/Sockets/Private/ServerTOC.cpp
--- a/Engine/Source/Runtime/Sockets/Private/ServerTOC.cpp
+++ b/Engine/Source/Runtime/Sockets/Private/ServerTOC.cpp
@@ -34,22 +34,34 @@ void FServerTOC::AddFileOrDirectory(const FString& Filename, const FDateTime& Ti
 	ServerPathDirectory->Add(Filename, Timestamp);
 }
 
-FDateTime* FServerTOC::FindFile(const FString& Filename)
+int32 FServerTOC::RemoveFileOrDirectory(const FString& Filename )
 {
-	FDateTime* Result = NULL;
+	const FString Path = FPaths::GetPath( Filename );
+	FDirectory* ServerPathDirectory = FindDirectory(Path);
+	if ( ServerPathDirectory != NULL )
+	{
+		return ServerPathDirectory->Remove(Filename);
+	}
+	return 0;
+}
+
+
+const FDateTime* FServerTOC::FindFile(const FString& Filename) const
+{
+	const FDateTime* Result = NULL;
 	// Find a directory first.
 	const FString Path = FPaths::GetPath(Filename);
-	FDirectory** FoundDirectory = Directories.Find(Path);
+	const FDirectory*const* FoundDirectory = Directories.Find(Path);
 	if (FoundDirectory != NULL)
 	{
 		// Find a file inside this directory.
-		FDirectory* Directory = *FoundDirectory;
+		const FDirectory* Directory = *FoundDirectory;
 		Result = Directory->Find(Filename);		
 	}
 	return Result;
 }
 
-FServerTOC::FDirectory* FServerTOC::FindDirectory(const FString& Directory)
+FServerTOC::FDirectory* FServerTOC::FindDirectory(const FString& Directory) 
 {
 	FDirectory** FoundDirectory = Directories.Find(Directory);
 	if (FoundDirectory != NULL)
@@ -60,4 +72,17 @@ FServerTOC::FDirectory* FServerTOC::FindDirectory(const FString& Directory)
 	{
 		return NULL;
 	}
+}
+
+const FServerTOC::FDirectory* FServerTOC::FindDirectory(const FString& Directory) const
+{
+	const FDirectory*const* FoundDirectory = Directories.Find(Directory);
+	if (FoundDirectory != NULL)
+	{
+		return *FoundDirectory;
+	}
+	else
+	{
+		return NULL;
+	}
 }	
diff --git a/Engine/Source/Runtime/Sockets/Public/NetworkMessage.h b/Engine/Source/Runtime/Sockets/Public/NetworkMessage.h
--- a/Engine/Source/Runtime/Sockets/Public/NetworkMessage.h
+++ b/Engine/Source/Runtime/Sockets/Public/NetworkMessage.h
@@ -5,6 +5,7 @@
 #include "CoreMinimal.h"
 #include "Serialization/BufferArchive.h"
 #include "Serialization/ArrayReader.h"
+#include "Misc/EnumClassFlags.h"
 
 enum
 { 
@@ -59,6 +60,18 @@ namespace NFS_Channels
 	};
 }
 
+
+
+enum class EConnectionFlags : uint8
+{
+	None = 0x00000000,
+	Streaming = 0x00000001,
+	PreCookedIterative = 0x00000002,
+};
+ENUM_CLASS_FLAGS(EConnectionFlags);
+
+
+
 /**
 * Simple abstraction for sockets that allows FNFSMessageHeader to use either an ordinary socket or a mutichannel socket
 **/
diff --git a/Engine/Source/Runtime/Sockets/Public/ServerTOC.h b/Engine/Source/Runtime/Sockets/Public/ServerTOC.h
--- a/Engine/Source/Runtime/Sockets/Public/ServerTOC.h
+++ b/Engine/Source/Runtime/Sockets/Public/ServerTOC.h
@@ -32,13 +32,25 @@ struct SOCKETS_API FServerTOC
 	 * @param Filename File name to find.
 	 * @return Pointer to a timestamp if the file was found, NULL otherwise.
 	 */
-	FDateTime* FindFile(const FString& Filename);
+	const FDateTime* FindFile(const FString& Filename) const;
 
 	/**
 	 * Finds a directory in TOC.
 	 *
 	 * @param Directory Directory to find.
-	 * @return Pointer to a timestamp if the directory was found, NULL otherwise.
+	 * @return Pointer to a FDirectory if the directory was found, NULL otherwise.
+	 */
+	const FDirectory* FindDirectory(const FString& Directory) const;
+
+	/**
+	 * Finds a directory in TOC non const version used internally 
+	 * see FindDirectory
+	 *
+	 * @param Directory Directory to find
+	 * @return pointer to a FDirectory if the directory was found, null otherwise
 	 */
 	FDirectory* FindDirectory(const FString& Directory);
+
+
+	int32 RemoveFileOrDirectory(const FString& Filename);
 };
diff --git a/Engine/Source/Runtime/StreamingFile/Private/StreamingNetworkPlatformFile.cpp b/Engine/Source/Runtime/StreamingFile/Private/StreamingNetworkPlatformFile.cpp
--- a/Engine/Source/Runtime/StreamingFile/Private/StreamingNetworkPlatformFile.cpp
+++ b/Engine/Source/Runtime/StreamingFile/Private/StreamingNetworkPlatformFile.cpp
@@ -325,7 +325,7 @@ bool FStreamingNetworkPlatformFile::InitializeInternal(IPlatformFile* Inner, con
 
 	// Send the filenames and timestamps to the server.
 	FNetworkFileArchive Payload(NFS_Messages::GetFileList);
-	FillGetFileList(Payload, true);
+	FillGetFileList(Payload);
 
 	// Send the directories over, and wait for a response.
 	FArrayReader Response;
diff --git a/Engine/Source/Runtime/StreamingFile/Public/StreamingNetworkPlatformFile.h b/Engine/Source/Runtime/StreamingFile/Public/StreamingNetworkPlatformFile.h
--- a/Engine/Source/Runtime/StreamingFile/Public/StreamingNetworkPlatformFile.h
+++ b/Engine/Source/Runtime/StreamingFile/Public/StreamingNetworkPlatformFile.h
@@ -131,7 +131,11 @@ class STREAMINGFILE_API FStreamingNetworkPlatformFile
 public:
 
 	/** Default Constructor */
-	FStreamingNetworkPlatformFile() { };
+	FStreamingNetworkPlatformFile() 
+	{
+		HeartbeatFrequency = -1.0f;
+		ConnectionFlags |= EConnectionFlags::Streaming;
+	}
 
 	/** Virtual destructor */
 	virtual ~FStreamingNetworkPlatformFile();
diff --git a/Engine/Source/Runtime/Windows/D3D11RHI/Private/Windows/WindowsD3D11Device.cpp b/Engine/Source/Runtime/Windows/D3D11RHI/Private/Windows/WindowsD3D11Device.cpp
--- a/Engine/Source/Runtime/Windows/D3D11RHI/Private/Windows/WindowsD3D11Device.cpp
+++ b/Engine/Source/Runtime/Windows/D3D11RHI/Private/Windows/WindowsD3D11Device.cpp
@@ -561,13 +561,21 @@ static bool SupportsHDROutput(FD3D11DynamicRHI* D3DRHI)
 void FD3D11DynamicRHIModule::StartupModule()
 {	
 #if NV_AFTERMATH
-	FString AftermathBinariesRoot = FPaths::EngineDir() / TEXT("Binaries/ThirdParty/NVIDIA/NVaftermath/Win64/");
-	if (LoadLibraryW(*(AftermathBinariesRoot + "GFSDK_Aftermath_Lib.dll")) == nullptr)
-	{
-		UE_LOG(LogD3D11RHI, Log, TEXT("Failed to load GFSDK_Aftermath_Lib.dll"));
-		GDX11NVAfterMathEnabled = 0;
-		return;
-	}
+	// Note - can't check device type here, we'll check for that before actually initializing Aftermath
+
+		FString AftermathBinariesRoot = FPaths::EngineDir() / TEXT("Binaries/ThirdParty/NVIDIA/NVaftermath/Win64/");
+		if (LoadLibraryW(*(AftermathBinariesRoot + "GFSDK_Aftermath_Lib.dll")) == nullptr)
+		{
+			UE_LOG(LogD3D11RHI, Warning, TEXT("Failed to load GFSDK_Aftermath_Lib.dll"));
+			GDX11NVAfterMathEnabled = 0;
+			return;
+		}
+		else
+		{
+			UE_LOG(LogD3D11RHI, Log, TEXT("Aftermath initialized"));
+			GDX11NVAfterMathEnabled = 1;
+		}
+	
 #endif
 }
 
@@ -1080,16 +1088,9 @@ void FD3D11DynamicRHI::InitD3DDevice()
 				GSupportsTimestampRenderQueries = false;
 			}
 		}
-
-#if NV_AFTERMATH
-		if (!IsRHIDeviceNVIDIA())
-		{
-			GDX11NVAfterMathEnabled = 0;
-		}
-#endif
 		
 #if PLATFORM_DESKTOP
-		if( IsRHIDeviceNVIDIA() )
+		if (IsRHIDeviceNVIDIA())
 		{
 			GSupportsDepthBoundsTest = true;
 			NV_GET_CURRENT_SLI_STATE SLICaps;
@@ -1108,15 +1109,41 @@ void FD3D11DynamicRHI::InitD3DDevice()
 			{
 				UE_LOG(LogD3D11RHI, Log, TEXT("NvAPI_D3D_GetCurrentSLIState failed: 0x%x"), (int32)SLIStatus);
 			}
+		}
+		else if( IsRHIDeviceAMD() )
+		{
+			// The AMD shader extensions are currently unused in UE4, but we have to set the associated UAV slot
+			// to something in the call below (default is 7, so just use that)
+			const uint32 AmdShaderExtensionUavSlot = 7;
+
+			// Initialize AGS's driver extensions
+			uint32 AmdSupportedExtensionFlags = 0;
+			auto AmdAgsResult = agsDriverExtensionsDX11_Init(AmdAgsContext, Direct3DDevice, AmdShaderExtensionUavSlot, &AmdSupportedExtensionFlags);
+			if (AmdAgsResult == AGS_SUCCESS && (AmdSupportedExtensionFlags & AGS_DX11_EXTENSION_DEPTH_BOUNDS_TEST) != 0)
+			{
+				GSupportsDepthBoundsTest = true;
+			}
+		}
 
 #if NV_AFTERMATH
+		// Two ways to enable aftermath, command line or the r.GPUCrashDebugging variable
+		// Note: If intending to change this please alert game teams who use this for user support.
+		if (FParse::Param(FCommandLine::Get(), TEXT("gpucrashdebugging")))
+		{
+			GDX11NVAfterMathEnabled = true;
+		}
+		else
+		{
 			static IConsoleVariable* GPUCrashDebugging = IConsoleManager::Get().FindConsoleVariable(TEXT("r.GPUCrashDebugging"));
 			if (GPUCrashDebugging)
 			{
 				GDX11NVAfterMathEnabled = GPUCrashDebugging->GetInt();
 			}
-
-			if (GDX11NVAfterMathEnabled)
+		}
+			
+		if (GDX11NVAfterMathEnabled)
+		{
+			if (IsRHIDeviceNVIDIA())
 			{
 				auto Result = GFSDK_Aftermath_DX11_Initialize(GFSDK_Aftermath_Version_API, Direct3DDevice);
 				if (Result == GFSDK_Aftermath_Result_Success)
@@ -1134,29 +1161,20 @@ void FD3D11DynamicRHI::InitD3DDevice()
 					GDX11NVAfterMathEnabled = 0;
 				}
 			}
-#endif
-		}
-		else if( IsRHIDeviceAMD() )
-		{
-			// The AMD shader extensions are currently unused in UE4, but we have to set the associated UAV slot
-			// to something in the call below (default is 7, so just use that)
-			const uint32 AmdShaderExtensionUavSlot = 7;
-
-			// Initialize AGS's driver extensions
-			uint32 AmdSupportedExtensionFlags = 0;
-			auto AmdAgsResult = agsDriverExtensionsDX11_Init(AmdAgsContext, Direct3DDevice, AmdShaderExtensionUavSlot, &AmdSupportedExtensionFlags);
-			if (AmdAgsResult == AGS_SUCCESS && (AmdSupportedExtensionFlags & AGS_DX11_EXTENSION_DEPTH_BOUNDS_TEST) != 0)
+			else
 			{
-				GSupportsDepthBoundsTest = true;
+				GDX11NVAfterMathEnabled = 0;
+				UE_LOG(LogD3D11RHI, Warning, TEXT("[Aftermath] Skipping aftermath initialization on non-Nvidia device"));
 			}
 		}
+#endif // NV_AFTERMATH
 
 		if (GDX11ForcedGPUs > 0)
 		{
 			GNumActiveGPUsForRendering = GDX11ForcedGPUs;
 			UE_LOG(LogD3D11RHI, Log, TEXT("r.DX11NumForcedGPUs forcing GNumActiveGPUsForRendering to: %i "), GDX11ForcedGPUs);
 		}
-#endif
+#endif // PLATFORM_DESKTOP
 
 		SetupAfterDeviceCreation();
 
