Copying //UE4/Dev-Sequencer to //UE4/Main (Source: //UE4/Dev-Sequencer @ 3617642)

diff --git a/Engine/Source/Runtime/MovieSceneCapture/Private/ImageSequenceProtocol.cpp b/Engine/Source/Runtime/MovieSceneCapture/Private/ImageSequenceProtocol.cpp
--- a/Engine/Source/Runtime/MovieSceneCapture/Private/ImageSequenceProtocol.cpp
+++ b/Engine/Source/Runtime/MovieSceneCapture/Private/ImageSequenceProtocol.cpp
@@ -1,6 +1,7 @@
 // Copyright 1998-2017 Epic Games, Inc. All Rights Reserved.
 
 #include "Protocols/ImageSequenceProtocol.h"
+
 #include "Misc/CommandLine.h"
 #include "Misc/FileHelper.h"
 #include "HAL/RunnableThread.h"
@@ -8,10 +9,10 @@
 #include "Modules/ModuleManager.h"
 #include "Async/Future.h"
 #include "Async/Async.h"
-
 #include "Templates/Casts.h"
 #include "MovieSceneCaptureSettings.h"
 
+
 void UBmpImageCaptureSettings::OnReleaseConfig(FMovieSceneCaptureSettings& InSettings)
 {
 	// Remove .{frame} if it exists. The "." before the {frame} is intentional because some media players denote frame numbers separated by "."
@@ -59,7 +60,8 @@ void UImageCaptureSettings::OnLoadConfig(FMovieSceneCaptureSettings& InSettings)
 }
 
 #if WITH_EDITOR
-#include "Interfaces/IImageWrapperModule.h"
+
+#include "IImageWrapperModule.h"
 
 namespace
 {
@@ -71,7 +73,7 @@ struct FImageFrameData : IFramePayload
 	FString Filename;
 };
 
-FImageSequenceProtocol::FImageSequenceProtocol(EImageFormat::Type InFormat)
+FImageSequenceProtocol::FImageSequenceProtocol(EImageFormat InFormat)
 {
 	Format = InFormat;
 	CompressionQuality = 100;
@@ -152,7 +154,7 @@ void FImageSequenceProtocol::AddFormatMappings(TMap<FString, FStringFormatArg>&
 }
 
 
-FImageCaptureThread::FImageCaptureThread(EImageFormat::Type InFormat, int32 InCompressionQuality)
+FImageCaptureThread::FImageCaptureThread(EImageFormat InFormat, int32 InCompressionQuality)
 {
 	Format = InFormat;
 	CompressionQuality = InCompressionQuality;
@@ -258,8 +260,8 @@ uint32 FImageCaptureThread::Run()
 
 			for (int32 FrameIndex = 0; FrameIndex < ProcessingFrames.Num(); ++FrameIndex)
 			{
-				FrameFutures.Add(Async<void>(EAsyncExecution::TaskGraph, [&, FrameIndex]{
-					IImageWrapperPtr ImageWrapper = ImageWrappers.Num() ? ImageWrappers[FrameIndex] : nullptr;
+				FrameFutures.Add(Async<void>(EAsyncExecution::TaskGraph, [&, FrameIndex] {
+					TSharedPtr<IImageWrapper> ImageWrapper = ImageWrappers.Num() ? ImageWrappers[FrameIndex] : nullptr;
 					WriteFrameToDisk(ProcessingFrames[FrameIndex], ImageWrapper);
 				}));
 			}
@@ -279,7 +281,7 @@ uint32 FImageCaptureThread::Run()
 	return 0;
 }
 
-void FImageCaptureThread::WriteFrameToDisk(FCapturedFrameData& Frame, IImageWrapperPtr ImageWrapper) const
+void FImageCaptureThread::WriteFrameToDisk(FCapturedFrameData& Frame, TSharedPtr<IImageWrapper> ImageWrapper) const
 {
 	const float Width = Frame.BufferSize.X;
 	const float Height = Frame.BufferSize.Y;
diff --git a/Engine/Source/Runtime/MovieSceneCapture/Public/Protocols/ImageSequenceProtocol.h b/Engine/Source/Runtime/MovieSceneCapture/Public/Protocols/ImageSequenceProtocol.h
--- a/Engine/Source/Runtime/MovieSceneCapture/Public/Protocols/ImageSequenceProtocol.h
+++ b/Engine/Source/Runtime/MovieSceneCapture/Public/Protocols/ImageSequenceProtocol.h
@@ -9,7 +9,7 @@
 #include "HAL/ThreadSafeBool.h"
 
 #if WITH_EDITOR
-	#include "Interfaces/IImageWrapper.h"
+	#include "IImageWrapper.h"
 #endif
 
 #include "FrameGrabber.h"
@@ -53,7 +53,7 @@ public:
 /** Single runnable thread used to dispatch captured frames */
 struct FImageCaptureThread : public FRunnable
 {
-	FImageCaptureThread(EImageFormat::Type InFormat, int32 InCompressionQuality);
+	FImageCaptureThread(EImageFormat InFormat, int32 InCompressionQuality);
 	~FImageCaptureThread();
 
 	void Add(FCapturedFrameData Frame);
@@ -65,7 +65,7 @@ struct FImageCaptureThread : public FRunnable
 
 private:
 
-	void WriteFrameToDisk(FCapturedFrameData& Frame, IImageWrapperPtr ImageWrapper) const;
+	void WriteFrameToDisk(FCapturedFrameData& Frame, TSharedPtr<IImageWrapper> ImageWrapper) const;
 
 private:
 	/** The thread itself */
@@ -80,16 +80,16 @@ private:
 	/** Set to false when the thread should terminate */
 	FThreadSafeBool bRunning;
 	/** The format we are writing out */
-	EImageFormat::Type Format;
+	EImageFormat Format;
 	/** Level of compression to apply to the image, between 1 (worst quality, best compression) and 100 (best quality, worst compression) */
 	int32 CompressionQuality;
 	/** array of image writer pointers for async writing */
-	TArray<IImageWrapperPtr> ImageWrappers;
+	TArray<TSharedPtr<IImageWrapper>> ImageWrappers;
 };
 
 struct MOVIESCENECAPTURE_API FImageSequenceProtocol : FFrameGrabberProtocol
 {
-	FImageSequenceProtocol(EImageFormat::Type InFormat);
+	FImageSequenceProtocol(EImageFormat InFormat);
 
 	/** ~FFrameGrabberProtocol implementation */
 	virtual bool Initialize(const FCaptureProtocolInitSettings& InSettings, const ICaptureProtocolHost& Host) override;
@@ -107,7 +107,7 @@ private:
 	/** Level of compression to apply to the image, between 1 (worst quality, best compression) and 100 (best quality, worst compression)*/
 	int32 CompressionQuality;
 	/** The format of the image to write out */
-	EImageFormat::Type Format;
+	EImageFormat Format;
 	/** Thread responsible for writing out frames to disk */
 	TUniquePtr<FImageCaptureThread> CaptureThread;
 };
diff --git a/Engine/Source/Runtime/MovieSceneTracks/Private/Evaluation/MovieSceneSpawnTemplate.cpp b/Engine/Source/Runtime/MovieSceneTracks/Private/Evaluation/MovieSceneSpawnTemplate.cpp
--- a/Engine/Source/Runtime/MovieSceneTracks/Private/Evaluation/MovieSceneSpawnTemplate.cpp
+++ b/Engine/Source/Runtime/MovieSceneTracks/Private/Evaluation/MovieSceneSpawnTemplate.cpp
@@ -69,7 +69,9 @@ struct FSpawnObjectToken : IMovieSceneExecutionToken
 				const UMovieSceneSequence* Sequence = Player.State.FindSequence(Operand.SequenceID);
 				if (Sequence)
 				{
-					Player.GetSpawnRegister().SpawnObject(Operand.ObjectBindingID, *Sequence->GetMovieScene(), Operand.SequenceID, Player);
+					UObject* SpawnedObject = Player.GetSpawnRegister().SpawnObject(Operand.ObjectBindingID, *Sequence->GetMovieScene(), Operand.SequenceID, Player);
+
+					Player.OnObjectSpawned(SpawnedObject, Operand);
 				}
 			}
 
diff --git a/Engine/Source/Runtime/NetworkFileSystem/NetworkFileSystem.Build.cs b/Engine/Source/Runtime/NetworkFileSystem/NetworkFileSystem.Build.cs
--- a/Engine/Source/Runtime/NetworkFileSystem/NetworkFileSystem.Build.cs
+++ b/Engine/Source/Runtime/NetworkFileSystem/NetworkFileSystem.Build.cs
@@ -6,45 +6,41 @@ namespace UnrealBuildTool.Rules
 	{
 		public NetworkFileSystem(ReadOnlyTargetRules Target) : base(Target)
 		{
-            PrivateIncludePaths.AddRange(
-	            new string[] {
-		            "Runtime/NetworkFileSystem/Private",
-		            "Runtime/NetworkFileSystem/Private/Simple",
-		            "Runtime/NetworkFileSystem/Private/Streaming",
-	            }
-	            );
+			PrivateIncludePaths.AddRange(
+				new string[] {
+					"Runtime/NetworkFileSystem/Private",
+					"Runtime/NetworkFileSystem/Private/Simple",
+					"Runtime/NetworkFileSystem/Private/Streaming",
+				});
 
-            PrivateDependencyModuleNames.AddRange(
-	            new string[]
-	            {
-                    "CoreUObject",
-		            "Projects",
-		            "SandboxFile",
-		            "TargetPlatform",
-	            }
-            );
+			PrivateDependencyModuleNames.AddRange(
+				new string[]
+				{
+					"CoreUObject",
+					"Projects",
+					"SandboxFile",
+					"TargetPlatform",
+				});
 
-            PublicIncludePaths.AddRange(
-	            new string[] {
-		            "Runtime/NetworkFileSystem/Public",
-		            "Runtime/NetworkFileSystem/Public/Interfaces",
-                    "Runtime/CoreUObject/Public/Interfaces",
-	            }
-            );
+			PublicIncludePaths.AddRange(
+				new string[] {
+					"Runtime/NetworkFileSystem/Public",
+					"Runtime/NetworkFileSystem/Public/Interfaces",
+					"Runtime/CoreUObject/Public/Interfaces",
+				});
 
-            PublicDependencyModuleNames.AddRange(
-	            new string[]
-		            {
-			            "Core",
-			            "Sockets",
-		            }
-		            );
+			PublicDependencyModuleNames.AddRange(
+				new string[] {
+					"Core",
+					"Sockets",
+				});
 
-            if (Target.Platform == UnrealTargetPlatform.Win64 || Target.Platform == UnrealTargetPlatform.Mac)
-            {
-                AddEngineThirdPartyPrivateStaticDependencies(Target, "OpenSSL", "libWebSockets", "zlib");
-                Definitions.Add("ENABLE_HTTP_FOR_NFS=1");
-            }
+			if ((Target.Platform == UnrealTargetPlatform.Win64) ||
+				(Target.Platform == UnrealTargetPlatform.Mac))
+			{
+				AddEngineThirdPartyPrivateStaticDependencies(Target, "OpenSSL", "libWebSockets", "zlib");
+				Definitions.Add("ENABLE_HTTP_FOR_NFS=1");
+			}
 			else
 			{
 				Definitions.Add("ENABLE_HTTP_FOR_NFS=0");
diff --git a/Engine/Source/Runtime/NetworkFileSystem/Private/NetworkFileServer.h b/Engine/Source/Runtime/NetworkFileSystem/Private/NetworkFileServer.h
--- a/Engine/Source/Runtime/NetworkFileSystem/Private/NetworkFileServer.h
+++ b/Engine/Source/Runtime/NetworkFileSystem/Private/NetworkFileServer.h
@@ -5,13 +5,14 @@
 #include "CoreMinimal.h"
 #include "HAL/ThreadSafeCounter.h"
 #include "HAL/Runnable.h"
-#include "Interfaces/INetworkFileServer.h"
-#include "Interfaces/INetworkFileSystemModule.h"
+#include "INetworkFileServer.h"
+#include "INetworkFileSystemModule.h"
 
 class FInternetAddr;
 class FSocket;
 class ITargetPlatform;
 
+
 /**
  * This class wraps the server thread and network connection
  */
diff --git a/Engine/Source/Runtime/NetworkFileSystem/Private/NetworkFileServerConnection.h b/Engine/Source/Runtime/NetworkFileSystem/Private/NetworkFileServerConnection.h
--- a/Engine/Source/Runtime/NetworkFileSystem/Private/NetworkFileServerConnection.h
+++ b/Engine/Source/Runtime/NetworkFileSystem/Private/NetworkFileServerConnection.h
@@ -3,12 +3,13 @@
 #pragma once
 
 #include "CoreMinimal.h"
-#include "Interfaces/INetworkFileSystemModule.h"
 #include "CoreMisc.h"
+#include "INetworkFileSystemModule.h"
 
 class FSandboxPlatformFile;
 class ITargetPlatform;
 
+
 /**
  * This class processes all incoming messages from the client.
  */
diff --git a/Engine/Source/Runtime/NetworkFileSystem/Private/NetworkFileServerHttp.h b/Engine/Source/Runtime/NetworkFileSystem/Private/NetworkFileServerHttp.h
--- a/Engine/Source/Runtime/NetworkFileSystem/Private/NetworkFileServerHttp.h
+++ b/Engine/Source/Runtime/NetworkFileSystem/Private/NetworkFileServerHttp.h
@@ -1,19 +1,13 @@
 // Copyright 1998-2017 Epic Games, Inc. All Rights Reserved.
 
-/*=============================================================================
-	NetworkFileServerHttp.h: Declares the NetworkFileServerHttp class.
-
-	This allows NFS to use a http server for serving Unreal Files.
-
-=============================================================================*/
 #pragma  once
 
 #include "CoreMinimal.h"
 #include "HAL/ThreadSafeCounter.h"
 #include "Misc/Guid.h"
 #include "HAL/Runnable.h"
-#include "Interfaces/INetworkFileServer.h"
-#include "Interfaces/INetworkFileSystemModule.h"
+#include "INetworkFileServer.h"
+#include "INetworkFileSystemModule.h"
 
 class FInternetAddr;
 class FNetworkFileServerClientConnectionHTTP;
@@ -22,8 +16,8 @@ class ITargetPlatform;
 #if ENABLE_HTTP_FOR_NFS
 
 #if PLATFORM_WINDOWS
-#include "WindowsHWrapper.h"
-#include "AllowWindowsPlatformTypes.h"
+	#include "WindowsHWrapper.h"
+	#include "AllowWindowsPlatformTypes.h"
 #endif
 
 /*
@@ -32,13 +26,13 @@ class ITargetPlatform;
 	The define will move the openssl define out of the way.
 */
 #define UI UI_ST
-THIRD_PARTY_INCLUDES_START
-#include "libwebsockets.h"
-THIRD_PARTY_INCLUDES_END
+	THIRD_PARTY_INCLUDES_START
+	#include "libwebsockets.h"
+	THIRD_PARTY_INCLUDES_END
 #undef UI
 
 #if PLATFORM_WINDOWS
-#include "HideWindowsPlatformTypes.h"
+	#include "HideWindowsPlatformTypes.h"
 #endif
 
 
diff --git a/Engine/Source/Runtime/NetworkFileSystem/Private/NetworkFileSystemModule.cpp b/Engine/Source/Runtime/NetworkFileSystem/Private/NetworkFileSystemModule.cpp
--- a/Engine/Source/Runtime/NetworkFileSystem/Private/NetworkFileSystemModule.cpp
+++ b/Engine/Source/Runtime/NetworkFileSystem/Private/NetworkFileSystemModule.cpp
@@ -3,7 +3,7 @@
 #include "CoreMinimal.h"
 #include "Misc/CommandLine.h"
 #include "Modules/ModuleManager.h"
-#include "Interfaces/INetworkFileSystemModule.h"
+#include "INetworkFileSystemModule.h"
 #include "NetworkFileSystemLog.h"
 #include "NetworkFileServer.h"
 #include "NetworkFileServerHttp.h"
diff --git a/Engine/Source/Runtime/NetworkFileSystem/Public/INetworkFileSystemModule.h b/Engine/Source/Runtime/NetworkFileSystem/Public/INetworkFileSystemModule.h
new file mode 100644
--- /dev/null
+++ b/Engine/Source/Runtime/NetworkFileSystem/Public/INetworkFileSystemModule.h
@@ -0,0 +1,140 @@
+// Copyright 1998-2017 Epic Games, Inc. All Rights Reserved.
+
+#pragma once
+
+#include "CoreMinimal.h"
+#include "Modules/ModuleInterface.h"
+
+class INetworkFileServer;
+
+/**
+ * Delegate type for handling file requests from a network client.
+ *
+ * The first parameter is the name of the requested file.
+ * The second parameter will hold the list of unsolicited files to send back.
+ */
+DECLARE_DELEGATE_ThreeParams(FFileRequestDelegate, const FString&, const FString&, TArray<FString>&);
+
+struct FShaderRecompileData
+{
+	FString PlatformName;
+	/** The platform to compile shaders for, corresponds to EShaderPlatform, but a value of -1 indicates to compile for all target shader platforms. */
+	int32 ShaderPlatform;
+	TArray<FString>* ModifiedFiles;
+	TArray<uint8>* MeshMaterialMaps;
+	TArray<FString> MaterialsToLoad;
+	TArray<uint8> SerializedShaderResources;
+	bool bCompileChangedShaders;
+
+	FShaderRecompileData() :
+		ShaderPlatform(-1),
+		bCompileChangedShaders(true)
+	{}
+
+	FShaderRecompileData& operator=(const FShaderRecompileData& Other)
+	{
+		PlatformName = Other.PlatformName;
+		ShaderPlatform = Other.ShaderPlatform;
+		ModifiedFiles = Other.ModifiedFiles;
+		MeshMaterialMaps = Other.MeshMaterialMaps;
+		MaterialsToLoad = Other.MaterialsToLoad;
+		SerializedShaderResources = Other.SerializedShaderResources;
+		bCompileChangedShaders = Other.bCompileChangedShaders;
+
+		return *this;
+	}
+};
+
+
+/**
+ * Delegate type for handling shader recompilation requests from a network client.
+ */
+DECLARE_DELEGATE_OneParam(FRecompileShadersDelegate, const FShaderRecompileData&);
+
+/**
+ * Delegate which returns an override for the sandbox path
+ */
+DECLARE_DELEGATE_RetVal( FString, FSandboxPathDelegate);
+
+/**
+ * Delegate which is called when an outside system modifies a file
+ */
+DECLARE_MULTICAST_DELEGATE_OneParam( FOnFileModifiedDelegate, const FString& );
+
+
+/**
+ * Delegate which is called when a new connection is made to a file server client
+ * 
+ * @param 1 Version string
+ * @param 2 Platform name
+ * @return return false if the connection should be destroyed
+ */
+DECLARE_DELEGATE_RetVal_TwoParams( bool, FNewConnectionDelegate, const FString&, const FString& );
+
+
+/**
+ * Delegate which returns a list of all the files which should already be deployed to the devkit
+ *
+ * @param 1 IN, Platform to get precooked file list
+ * @param 2 OUT, list of precooked files 
+ */
+typedef TMap<FString,FDateTime> FFileTimeMap;
+DECLARE_DELEGATE_TwoParams( FInitialPrecookedListDelegate, const FString&, FFileTimeMap& );
+
+
+
+// container struct for delegates which the network file system uses
+struct FNetworkFileDelegateContainer
+{
+public:
+	FNetworkFileDelegateContainer() : 
+		NewConnectionDelegate(nullptr), 
+		InitialPrecookedListDelegate(nullptr),
+		SandboxPathOverrideDelegate(nullptr),
+		FileRequestDelegate(nullptr),
+		RecompileShadersDelegate(nullptr),
+		OnFileModifiedCallback(nullptr)
+	{}
+	FNewConnectionDelegate NewConnectionDelegate; 
+	FInitialPrecookedListDelegate InitialPrecookedListDelegate;
+	FSandboxPathDelegate SandboxPathOverrideDelegate;
+	FFileRequestDelegate FileRequestDelegate;
+	FRecompileShadersDelegate RecompileShadersDelegate;
+
+	FOnFileModifiedDelegate* OnFileModifiedCallback; // this is called from other systems to notify the network file system that a file has been modified hence the terminology callback
+};
+
+
+enum ENetworkFileServerProtocol
+{
+	NFSP_Tcp,
+	NFSP_Http,
+};
+
+/**
+ * Interface for network file system modules.
+ */
+class INetworkFileSystemModule
+	: public IModuleInterface
+{
+public:
+
+	/**
+	 * Creates a new network file server.
+	 *
+	 * @param InPort The port number to bind to (-1 = default port, 0 = any available port).
+	 * @param Streaming Whether it should be a streaming server.
+	 * @param InFileRequestDelegate An optional delegate to be invoked when a file is requested by a client.
+	 * @param InRecompileShadersDelegate An optional delegate to be invoked when shaders need to be recompiled.
+	 *
+	 * @return The new file server, or nullptr if creation failed.
+	 */
+	virtual INetworkFileServer* CreateNetworkFileServer( bool bLoadTargetPlatforms, int32 Port = -1, FNetworkFileDelegateContainer InNetworkFileDelegateContainer = FNetworkFileDelegateContainer(), const ENetworkFileServerProtocol Protocol = NFSP_Tcp ) const = 0;
+
+public:
+
+	/**
+	 * Virtual destructor.
+	 */
+	virtual ~INetworkFileSystemModule( ) { }
+};
diff --git a/Engine/Source/Runtime/Networking/Public/Common/TcpListener.h b/Engine/Source/Runtime/Networking/Public/Common/TcpListener.h
--- a/Engine/Source/Runtime/Networking/Public/Common/TcpListener.h
+++ b/Engine/Source/Runtime/Networking/Public/Common/TcpListener.h
@@ -120,7 +120,6 @@ public:
 	 * To temporarily disable accepting connections, use the Enable() and Disable() methods.
 	 *
 	 * @return The delegate.
-	 * @see Enable, Disable
 	 */
 	FOnTcpListenerConnectionAccepted& OnConnectionAccepted()
 	{
@@ -138,10 +137,8 @@ public:
 			Socket = FTcpSocketBuilder(TEXT("FTcpListener server"))
 				.AsReusable()
 				.BoundToEndpoint(Endpoint)
-				.Listening(8);
-
-			int32 NewSize = 0;
-			Socket->SetReceiveBufferSize(2*1024*1024, NewSize);
+				.Listening(8)
+				.WithSendBufferSize(2 * 1024 * 1024);
 		}
 
 		return (Socket != nullptr);
diff --git a/Engine/Source/Runtime/Online/ImageDownload/Private/WebImage.cpp b/Engine/Source/Runtime/Online/ImageDownload/Private/WebImage.cpp
--- a/Engine/Source/Runtime/Online/ImageDownload/Private/WebImage.cpp
+++ b/Engine/Source/Runtime/Online/ImageDownload/Private/WebImage.cpp
@@ -1,12 +1,16 @@
 // Copyright 1998-2017 Epic Games, Inc. All Rights Reserved.
 
 #include "WebImage.h"
-#include "ImageDownloadPrivate.h"
-#include "Interfaces/IImageWrapperModule.h"
+
+#include "IImageWrapper.h"
+#include "IImageWrapperModule.h"
 #include "HttpModule.h"
 #include "Modules/ModuleManager.h"
 #include "Styling/CoreStyle.h"
 
+#include "ImageDownloadPrivate.h"
+
+
 FWebImage::FWebImage()
 : StandInBrush(FCoreStyle::Get().GetDefaultBrush())
 , bDownloadSuccess(false)
@@ -110,7 +114,7 @@ bool FWebImage::ProcessHttpResponse(const FString& RequestUrl, FHttpResponsePtr
 
 	// Look at the signature of the downloaded image to detect image type. (and ignore the content type header except for error reporting)
 	const TArray<uint8>& Content = HttpResponse->GetContent();
-	EImageFormat::Type ImageFormat = ImageWrapperModule.DetectImageFormat(Content.GetData(), Content.Num());
+	EImageFormat ImageFormat = ImageWrapperModule.DetectImageFormat(Content.GetData(), Content.Num());
 
 	if (ImageFormat == EImageFormat::Invalid)
 	{
@@ -119,7 +123,7 @@ bool FWebImage::ProcessHttpResponse(const FString& RequestUrl, FHttpResponsePtr
 		return false;
 	}
 
-	IImageWrapperPtr ImageWrapper = ImageWrapperModule.CreateImageWrapper(ImageFormat);
+	TSharedPtr<IImageWrapper> ImageWrapper = ImageWrapperModule.CreateImageWrapper(ImageFormat);
 	if (!ImageWrapper.IsValid())
 	{
 		UE_LOG(LogImageDownload, Error, TEXT("Image Download: Unable to make image wrapper for image format %d"), (int32)ImageFormat);
diff --git a/Engine/Source/Runtime/Online/XMPP/Private/XmppStrophe/XmppPingStrophe.h b/Engine/Source/Runtime/Online/XMPP/Private/XmppStrophe/XmppPingStrophe.h
--- a/Engine/Source/Runtime/Online/XMPP/Private/XmppStrophe/XmppPingStrophe.h
+++ b/Engine/Source/Runtime/Online/XMPP/Private/XmppStrophe/XmppPingStrophe.h
@@ -4,7 +4,7 @@
 
 #include "CoreMinimal.h"
 #include "XmppConnection.h"
-
+#include "Containers/Queue.h"
 #include "Containers/Ticker.h"
 
 class FXmppConnectionStrophe;
diff --git a/Engine/Source/Runtime/OpenGLDrv/Private/Android/AndroidESDeferredOpenGL.h b/Engine/Source/Runtime/OpenGLDrv/Private/Android/AndroidESDeferredOpenGL.h
--- a/Engine/Source/Runtime/OpenGLDrv/Private/Android/AndroidESDeferredOpenGL.h
+++ b/Engine/Source/Runtime/OpenGLDrv/Private/Android/AndroidESDeferredOpenGL.h
@@ -530,6 +530,19 @@ struct FAndroidESDeferredOpenGL : public FOpenGLESDeferred
 		return false;
 	}
 
+	enum class EImageExternalType : uint8
+	{
+		None,
+		ImageExternal100,
+		ImageExternal300,
+		ImageExternalESSL300
+	};
+
+	static FORCEINLINE EImageExternalType GetImageExternalType()
+	{
+		return EImageExternalType::None;
+	}
+
 	static FORCEINLINE GLuint64 GetTextureSamplerHandle(GLuint Texture, GLuint Sampler)
 	{
 		return glGetTextureSamplerHandleNV(Texture, Sampler);
diff --git a/Engine/Source/Runtime/OpenGLDrv/Private/Android/AndroidOpenGL.cpp b/Engine/Source/Runtime/OpenGLDrv/Private/Android/AndroidOpenGL.cpp
--- a/Engine/Source/Runtime/OpenGLDrv/Private/Android/AndroidOpenGL.cpp
+++ b/Engine/Source/Runtime/OpenGLDrv/Private/Android/AndroidOpenGL.cpp
@@ -408,6 +408,7 @@ bool FAndroidOpenGL::bSupportsInstancing = false;
 bool FAndroidOpenGL::bHasHardwareHiddenSurfaceRemoval = false;
 bool FAndroidOpenGL::bSupportsMobileMultiView = false;
 bool FAndroidOpenGL::bSupportsImageExternal = false;
+FAndroidOpenGL::EImageExternalType FAndroidOpenGL::ImageExternalType = FAndroidOpenGL::EImageExternalType::None;
 GLint FAndroidOpenGL::MaxMSAASamplesTileMem = 1;
 
 FAndroidOpenGL::EFeatureLevelSupport FAndroidOpenGL::CurrentFeatureLevelSupport = FAndroidOpenGL::EFeatureLevelSupport::Invalid;
@@ -470,17 +471,69 @@ void FAndroidOpenGL::ProcessExtensions(const FString& ExtensionsString)
 		MaxMSAASamplesTileMem = 1;
 	}
 
-	if (ExtensionsString.Contains(TEXT("GL_OES_EGL_image_external")))
-	{
-		bSupportsImageExternal = true;
-		UE_LOG(LogRHI, Log, TEXT("Image external enabled."));
-	}
-
 	bSupportsETC2 = bES30Support;
 	bUseES30ShadingLanguage = bES30Support;
 
 	FString RendererString = FString(ANSI_TO_TCHAR((const ANSICHAR*)glGetString(GL_RENDERER)));
 
+	// Common GPU types
+	const bool bIsNvidiaBased = RendererString.Contains(TEXT("NVIDIA"));
+	const bool bIsPoverVRBased = RendererString.Contains(TEXT("PowerVR"));
+	const bool bIsAdrenoBased = RendererString.Contains(TEXT("Adreno"));
+	const bool bIsMaliBased = RendererString.Contains(TEXT("Mali"));
+
+	// Check for external image support for different ES versions
+	bool bHasImageExternal = ExtensionsString.Contains(TEXT("GL_OES_EGL_image_external ")) || ExtensionsString.EndsWith(TEXT("GL_OES_EGL_image_external"));
+	bool bHasImageExternalESSL3 = ExtensionsString.Contains(TEXT("OES_EGL_image_external_essl3"));
+	if (bHasImageExternal || bHasImageExternalESSL3)
+	{
+		bSupportsImageExternal = true;
+		ImageExternalType = EImageExternalType::ImageExternal100;
+		if (bUseES30ShadingLanguage)
+		{
+			if (bHasImageExternalESSL3)
+			{
+				ImageExternalType = EImageExternalType::ImageExternalESSL300;
+			}
+			else
+			{
+				// Adreno 5xx can do essl3 even without extension in list
+				if (bIsAdrenoBased && RendererString.Contains(TEXT("(TM) 5")))
+				{
+					ImageExternalType = EImageExternalType::ImageExternalESSL300;
+				}
+			}
+		}
+		if (bIsNvidiaBased)
+		{
+			// Nvidia needs version 100 even though it supports ES3
+			ImageExternalType = EImageExternalType::ImageExternal100;
+		}
+
+		switch (ImageExternalType)
+		{
+			case EImageExternalType::None:
+				bSupportsImageExternal = false;
+				break;
+
+			case EImageExternalType::ImageExternal100:
+				UE_LOG(LogRHI, Log, TEXT("Image external enabled: ImageExternal100"));
+				break;
+
+			case EImageExternalType::ImageExternal300:
+				UE_LOG(LogRHI, Log, TEXT("Image external enabled: ImageExternal300"));
+				break;
+
+			case EImageExternalType::ImageExternalESSL300:
+				UE_LOG(LogRHI, Log, TEXT("Image external enabled: ImageExternalESSL300"));
+				break;
+
+			default:
+				bSupportsImageExternal = false;
+				UE_LOG(LogRHI, Log, TEXT("Image external disabled; unknown type"));
+		}
+	}
+
 	if (RendererString.Contains(TEXT("SGX 540")))
 	{
 		UE_LOG(LogRHI, Warning, TEXT("Disabling support for GL_OES_packed_depth_stencil on SGX 540"));
@@ -488,14 +541,12 @@ void FAndroidOpenGL::ProcessExtensions(const FString& ExtensionsString)
 		bRequiresTexture2DPrecisionHack = true;
 	}
 
-	const bool bIsPoverVRBased = RendererString.Contains(TEXT("PowerVR"));
 	if (bIsPoverVRBased)
 	{
 		bHasHardwareHiddenSurfaceRemoval = true;
 		UE_LOG(LogRHI, Log, TEXT("Enabling support for Hidden Surface Removal on PowerVR"));
 	}
 
-	const bool bIsAdrenoBased = RendererString.Contains(TEXT("Adreno"));
 	if (bIsAdrenoBased)
 	{
 		// This is to avoid a bug in Adreno drivers that define GL_ARM_shader_framebuffer_fetch_depth_stencil even when device does not support this extension
diff --git a/Engine/Source/Runtime/OpenGLDrv/Private/Android/AndroidOpenGL.h b/Engine/Source/Runtime/OpenGLDrv/Private/Android/AndroidOpenGL.h
--- a/Engine/Source/Runtime/OpenGLDrv/Private/Android/AndroidOpenGL.h
+++ b/Engine/Source/Runtime/OpenGLDrv/Private/Android/AndroidOpenGL.h
@@ -467,6 +467,16 @@ struct FAndroidOpenGL : public FOpenGLES2
 		return bUseES30ShadingLanguage;
 	}
 
+	enum class EImageExternalType : uint8
+	{
+		None,
+		ImageExternal100,
+		ImageExternal300,
+		ImageExternalESSL300
+	};
+
+	static FORCEINLINE EImageExternalType GetImageExternalType() { return ImageExternalType; }
+
 	// For now, framebuffer fetch is disabled with ES3.1
 	static FORCEINLINE bool SupportsShaderFramebufferFetch() { return !IsES31Usable() && FOpenGLES2::SupportsShaderFramebufferFetch(); }
 	static FORCEINLINE bool SupportsShaderDepthStencilFetch() { return !IsES31Usable() && FOpenGLES2::SupportsShaderDepthStencilFetch(); }
@@ -508,6 +518,9 @@ struct FAndroidOpenGL : public FOpenGLES2
 	/** Whether device supports image external */
 	static bool bSupportsImageExternal;
 
+	/** Type of image external supported */
+	static EImageExternalType ImageExternalType;
+
 	/** Maximum number of MSAA samples supported on chip in tile memory, or 1 if not available */
 	static GLint MaxMSAASamplesTileMem;
 
diff --git a/Engine/Source/Runtime/OpenGLDrv/Private/OpenGLDebugFrameDump.cpp b/Engine/Source/Runtime/OpenGLDrv/Private/OpenGLDebugFrameDump.cpp
--- a/Engine/Source/Runtime/OpenGLDrv/Private/OpenGLDebugFrameDump.cpp
+++ b/Engine/Source/Runtime/OpenGLDrv/Private/OpenGLDebugFrameDump.cpp
@@ -1,9 +1,5 @@
 // Copyright 1998-2017 Epic Games, Inc. All Rights Reserved.
 
-/*=============================================================================
-	OpenGLDebugFrameDump.cpp: Implementation of debug code that allows creating detailed summaries of OpenGL pipeline state for all draws in a single frame.
-=============================================================================*/
-
 #include "CoreMinimal.h"
 #include "HAL/FileManager.h"
 #include "Misc/FileHelper.h"
@@ -16,30 +12,32 @@
 #define USE_COMPRESSED_PNG_INSTEAD_OF_BMP_FOR_CONTENT_OUTPUT 1
 
 #if USE_COMPRESSED_PNG_INSTEAD_OF_BMP_FOR_CONTENT_OUTPUT
-// For PNG compression
-#include "Interfaces/IImageWrapperModule.h"
-const GLenum TextureOutputFormat = GL_RGBA;
+	// For PNG compression
+	#include "IImageWrapper.h"
+	#include "IImageWrapperModule.h"
+	const GLenum TextureOutputFormat = GL_RGBA;
 #else
-const GLenum TextureOutputFormat = GL_BGRA;
+	const GLenum TextureOutputFormat = GL_BGRA;
 #endif
 
 #define DEBUG_GL_ERRORS_CAUSED_BY_THIS_CODE 1
 
 #if DEBUG_GL_ERRORS_CAUSED_BY_THIS_CODE
-#define ASSERT_NO_GL_ERROR()	check( glGetError() == GL_NO_ERROR )
+	#define ASSERT_NO_GL_ERROR()	check( glGetError() == GL_NO_ERROR )
 #else
-#define ASSERT_NO_GL_ERROR()
+	#define ASSERT_NO_GL_ERROR()
 #endif
 
 extern bool GDisableOpenGLDebugOutput;
 
+
 #if USE_COMPRESSED_PNG_INSTEAD_OF_BMP_FOR_CONTENT_OUTPUT
 
 void appCreatePNGWithAlpha( const TCHAR* File, int32 Width, int32 Height, FColor* Data, IFileManager* FileManager = &IFileManager::Get() )
 {
 	// We assume all resources are png for now.
 	IImageWrapperModule& ImageWrapperModule = FModuleManager::LoadModuleChecked<IImageWrapperModule>( FName("ImageWrapper") );
-	IImageWrapperPtr ImageWrapper = ImageWrapperModule.CreateImageWrapper( EImageFormat::PNG );
+	TSharedPtr<IImageWrapper> ImageWrapper = ImageWrapperModule.CreateImageWrapper( EImageFormat::PNG );
 	if ( ImageWrapper.IsValid() && ImageWrapper->SetRaw( Data, 4 * Width * Height, Width, Height, ERGBFormat::RGBA, 8 ) )
 	{
 		FArchive* Ar = FileManager->CreateFileWriter( File );
diff --git a/Engine/Source/Runtime/OpenGLDrv/Private/OpenGLShaders.cpp b/Engine/Source/Runtime/OpenGLDrv/Private/OpenGLShaders.cpp
--- a/Engine/Source/Runtime/OpenGLDrv/Private/OpenGLShaders.cpp
+++ b/Engine/Source/Runtime/OpenGLDrv/Private/OpenGLShaders.cpp
@@ -691,13 +691,24 @@ void OPENGLDRV_API GLSLToDeviceCompatibleGLSL(FAnsiCharArray& GlslCodeOriginal,
 
 	// Whether we need to emit texture external code or not.
 	const bool bEmitTextureExternal = (FCStringAnsi::Strstr(GlslCodeOriginal.GetData(), "samplerExternalOES") != nullptr);
-	
+
+	bool bUseES30ShadingLanguage = Capabilities.bUseES30ShadingLanguage;
+
+#if PLATFORM_ANDROID
+	FAndroidOpenGL::EImageExternalType ImageExternalType = FAndroidOpenGL::GetImageExternalType();
+
+	if (ImageExternalType == FAndroidOpenGL::EImageExternalType::ImageExternal100)
+	{
+		bUseES30ShadingLanguage = false;
+	}
+#endif
+
 	if (Capabilities.TargetPlatform == EOpenGLShaderTargetPlatform::OGLSTP_Android || Capabilities.TargetPlatform == EOpenGLShaderTargetPlatform::OGLSTP_HTML5)
 	{
 		if (IsES2Platform(Capabilities.MaxRHIShaderPlatform) && !bES31)
 		{
 			// #version NNN has to be the first line in the file, so it has to be added before anything else.
-			if (Capabilities.bUseES30ShadingLanguage)
+			if (bUseES30ShadingLanguage)
 			{
 				AppendCString(GlslCode, "#version 300 es\n");
 			}
@@ -738,7 +749,26 @@ void OPENGLDRV_API GLSLToDeviceCompatibleGLSL(FAnsiCharArray& GlslCodeOriginal,
 		if (GSupportsImageExternal)
 		{
 			AppendCString(GlslCode, "\n\n");
-			AppendCString(GlslCode, "#extension GL_OES_EGL_image_external_essl3 : require\n");
+
+#if PLATFORM_ANDROID
+			switch (ImageExternalType)
+			{
+				case FAndroidOpenGL::EImageExternalType::ImageExternal100:
+					AppendCString(GlslCode, "#extension GL_OES_EGL_image_external : require\n");
+					break;
+
+				case FAndroidOpenGL::EImageExternalType::ImageExternal300:
+					AppendCString(GlslCode, "#extension GL_OES_EGL_image_external : require\n");
+					break;
+
+				case FAndroidOpenGL::EImageExternalType::ImageExternalESSL300:
+					// GL_OES_EGL_image_external_essl3 is only compatible with ES 3.x
+					AppendCString(GlslCode, "#extension GL_OES_EGL_image_external_essl3 : require\n");
+					break;
+			}
+#else
+			AppendCString(GlslCode, "#extension GL_OES_EGL_image_external : require\n");
+#endif
 			AppendCString(GlslCode, "\n\n");
 		}
 		else
@@ -851,7 +881,7 @@ void OPENGLDRV_API GLSLToDeviceCompatibleGLSL(FAnsiCharArray& GlslCodeOriginal,
 			}
 
 			// This #define fixes compiler errors on Android (which doesn't seem to support textureCubeLodEXT)
-			if (Capabilities.bUseES30ShadingLanguage)
+			if (bUseES30ShadingLanguage)
 			{
 				if (TypeEnum == GL_VERTEX_SHADER)
 				{
diff --git a/Engine/Source/Runtime/Portal/LauncherCheck/Private/LauncherCheckModule.cpp b/Engine/Source/Runtime/Portal/LauncherCheck/Private/LauncherCheckModule.cpp
--- a/Engine/Source/Runtime/Portal/LauncherCheck/Private/LauncherCheckModule.cpp
+++ b/Engine/Source/Runtime/Portal/LauncherCheck/Private/LauncherCheckModule.cpp
@@ -8,7 +8,8 @@
 #include "Misc/CommandLine.h"
 #include "Misc/Paths.h"
 #include "Modules/ModuleManager.h"
-#include "Interfaces/ILauncherCheckModule.h"
+#include "ILauncherCheckModule.h"
+
 
 #if defined(WITH_LAUNCHERCHECK) && WITH_LAUNCHERCHECK
 
@@ -16,11 +17,13 @@
 #include "ILauncherPlatform.h"
 #include "LauncherPlatformModule.h"
 
+
 /**
  * Log categories for LauncherCheck module
  */
 DEFINE_LOG_CATEGORY(LogLauncherCheck);
 
+
 /**
  * Implements the Launcher Check module.
  */
@@ -38,7 +41,9 @@ public:
 		return FParse::Param(FCommandLine::Get(), TEXT("NoEpicPortal")) == false;
 	}
 
-	// ILauncherCheckModule interface
+public:
+
+	//~ ILauncherCheckModule interface
 
 	virtual bool WasRanFromLauncher() const override
 	{
@@ -81,20 +86,15 @@ public:
 
 public:
 
-	// IModuleInterface interface
+	//~ IModuleInterface interface
 
-	virtual void StartupModule() override
-	{
-	}
-
-	virtual void ShutdownModule() override
-	{
-	}
+	virtual void StartupModule() override { }
+	virtual void ShutdownModule() override { }
 
 private:
 
 	/**
-	 * Return url encoded full path of currently running exe
+	 * Return url encoded full path of currently running executable.
 	 */
 	FString GetEncodedExePath() const
 	{
@@ -125,22 +125,26 @@ private:
 };
 
 
-#else
+#else //WITH_LAUNCHERCHECK
+
 
 class FLauncherCheckModule
 	: public ILauncherCheckModule
 {
 public:
 
-	virtual bool WasRanFromLauncher() const override { return true; }
-
-	virtual bool RunLauncher(ELauncherAction Action, FString Payload = FString()) const override { return false; }
+	virtual bool WasRanFromLauncher() const override
+	{
+		return true;
+	}
 
+	virtual bool RunLauncher(ELauncherAction Action, FString Payload = FString()) const override
+	{
+		return false;
+	}
 };
 
-
-#endif // WITH_LAUNCHERCHECK
-
-IMPLEMENT_MODULE(FLauncherCheckModule, LauncherCheck );
+#endif //WITH_LAUNCHERCHECK
 
 
+IMPLEMENT_MODULE(FLauncherCheckModule, LauncherCheck);
diff --git a/Engine/Source/Runtime/Portal/Rpc/PortalRpc.Build.cs b/Engine/Source/Runtime/Portal/Rpc/PortalRpc.Build.cs
--- a/Engine/Source/Runtime/Portal/Rpc/PortalRpc.Build.cs
+++ b/Engine/Source/Runtime/Portal/Rpc/PortalRpc.Build.cs
@@ -23,6 +23,7 @@ namespace UnrealBuildTool.Rules
             PrivateIncludePathModuleNames.AddRange(
                 new string[] {
                     "Messaging",
+                    "MessagingCommon",
                     "MessagingRpc",
                     "PortalServices",
                 }
diff --git a/Engine/Source/Runtime/Portal/Rpc/Private/PortalRpcLocator.cpp b/Engine/Source/Runtime/Portal/Rpc/Private/PortalRpcLocator.cpp
--- a/Engine/Source/Runtime/Portal/Rpc/Private/PortalRpcLocator.cpp
+++ b/Engine/Source/Runtime/Portal/Rpc/Private/PortalRpcLocator.cpp
@@ -6,8 +6,8 @@
 #include "Containers/Ticker.h"
 #include "Misc/EngineVersion.h"
 #include "IMessageContext.h"
-#include "Helpers/MessageEndpoint.h"
-#include "Helpers/MessageEndpointBuilder.h"
+#include "MessageEndpoint.h"
+#include "MessageEndpointBuilder.h"
 #include "PortalRpcDefines.h"
 #include "IPortalRpcLocator.h"
 #include "PortalRpcMessages.h"
diff --git a/Engine/Source/Runtime/Portal/Rpc/Private/PortalRpcResponder.cpp b/Engine/Source/Runtime/Portal/Rpc/Private/PortalRpcResponder.cpp
--- a/Engine/Source/Runtime/Portal/Rpc/Private/PortalRpcResponder.cpp
+++ b/Engine/Source/Runtime/Portal/Rpc/Private/PortalRpcResponder.cpp
@@ -2,8 +2,8 @@
 
 #include "PortalRpcResponder.h"
 #include "IMessageContext.h"
-#include "Helpers/MessageEndpoint.h"
-#include "Helpers/MessageEndpointBuilder.h"
+#include "MessageEndpoint.h"
+#include "MessageEndpointBuilder.h"
 #include "IPortalRpcResponder.h"
 #include "HAL/PlatformProcess.h"
 #include "IPortalRpcServer.h"
diff --git a/Engine/Source/Runtime/Portal/Rpc/Private/PortalRpcServer.cpp b/Engine/Source/Runtime/Portal/Rpc/Private/PortalRpcServer.cpp
--- a/Engine/Source/Runtime/Portal/Rpc/Private/PortalRpcServer.cpp
+++ b/Engine/Source/Runtime/Portal/Rpc/Private/PortalRpcServer.cpp
@@ -3,7 +3,7 @@
 #include "PortalRpcServer.h"
 #include "IPortalRpcServer.h"
 #include "PortalRpcMessages.h"
-#include "Helpers/MessageEndpoint.h"
+#include "MessageEndpoint.h"
 #include "MessageRpcServer.h"
 
 class FPortalRpcServerImpl
diff --git a/Engine/Source/Runtime/SessionServices/Private/SessionInstanceInfo.cpp b/Engine/Source/Runtime/SessionServices/Private/SessionInstanceInfo.cpp
--- a/Engine/Source/Runtime/SessionServices/Private/SessionInstanceInfo.cpp
+++ b/Engine/Source/Runtime/SessionServices/Private/SessionInstanceInfo.cpp
@@ -1,12 +1,14 @@
 // Copyright 1998-2017 Epic Games, Inc. All Rights Reserved.
 
 #include "SessionInstanceInfo.h"
-#include "SessionLogMessage.h"
-#include "HAL/PlatformProcess.h"
+
 #include "EngineServiceMessages.h"
-#include "Helpers/MessageEndpointBuilder.h"
+#include "HAL/PlatformProcess.h"
+#include "MessageEndpointBuilder.h"
 #include "SessionServiceMessages.h"
 
+#include "SessionLogMessage.h"
+
 
 /* FSessionInstanceInfo structors
  *****************************************************************************/
@@ -25,7 +27,7 @@ FSessionInstanceInfo::FSessionInstanceInfo(const FGuid& InInstanceId, const TSha
 /* FSessionInstanceInfo interface
  *****************************************************************************/
 
-void FSessionInstanceInfo::UpdateFromMessage(const FEngineServicePong& Message, const IMessageContextRef& Context)
+void FSessionInstanceInfo::UpdateFromMessage(const FEngineServicePong& Message, const TSharedRef<IMessageContext, ESPMode::ThreadSafe>& Context)
 {
 	if (Message.InstanceId != InstanceId)
 	{
@@ -41,7 +43,7 @@ void FSessionInstanceInfo::UpdateFromMessage(const FEngineServicePong& Message,
 }
 
 
-void FSessionInstanceInfo::UpdateFromMessage(const FSessionServicePong& Message, const IMessageContextRef& Context)
+void FSessionInstanceInfo::UpdateFromMessage(const FSessionServicePong& Message, const TSharedRef<IMessageContext, ESPMode::ThreadSafe>& Context)
 {
 	if (Message.InstanceId != InstanceId)
 	{
@@ -89,7 +91,7 @@ void FSessionInstanceInfo::Terminate()
 /* FSessionInstanceInfo callbacks
  *****************************************************************************/
 
-void FSessionInstanceInfo::HandleSessionLogMessage(const FSessionServiceLog& Message, const IMessageContextRef& Context)
+void FSessionInstanceInfo::HandleSessionLogMessage(const FSessionServiceLog& Message, const TSharedRef<IMessageContext, ESPMode::ThreadSafe>& Context)
 {
 	TSharedRef<FSessionLogMessage> LogMessage = MakeShareable(
 		new FSessionLogMessage(
diff --git a/Engine/Source/Runtime/SessionServices/Private/SessionInstanceInfo.h b/Engine/Source/Runtime/SessionServices/Private/SessionInstanceInfo.h
--- a/Engine/Source/Runtime/SessionServices/Private/SessionInstanceInfo.h
+++ b/Engine/Source/Runtime/SessionServices/Private/SessionInstanceInfo.h
@@ -6,7 +6,7 @@
 #include "Misc/Guid.h"
 #include "IMessageContext.h"
 #include "IMessageBus.h"
-#include "Helpers/MessageEndpoint.h"
+#include "MessageEndpoint.h"
 #include "ISessionInstanceInfo.h"
 #include "ISessionInfo.h"
 
@@ -44,7 +44,7 @@ public:
 	 * @param Message The message containing engine information.
 	 * @param Context The message context.
 	 */
-	void UpdateFromMessage(const FEngineServicePong& Message, const IMessageContextRef& Context);
+	void UpdateFromMessage(const FEngineServicePong& Message, const TSharedRef<IMessageContext, ESPMode::ThreadSafe>& Context);
 
 	/**
 	 * Updates this instance info with the data in the specified message.
@@ -52,7 +52,7 @@ public:
 	 * @param Message The message containing instance information.
 	 * @param Context The message context.
 	 */
-	void UpdateFromMessage(const FSessionServicePong& Message, const IMessageContextRef& Context);
+	void UpdateFromMessage(const FSessionServicePong& Message, const TSharedRef<IMessageContext, ESPMode::ThreadSafe>& Context);
 
 public:	
 
@@ -146,7 +146,7 @@ public:
 private:
 
 	/** Handles FSessionServiceLog messages. */
-	void HandleSessionLogMessage(const FSessionServiceLog& Message, const IMessageContextRef& Context);
+	void HandleSessionLogMessage(const FSessionServiceLog& Message, const TSharedRef<IMessageContext, ESPMode::ThreadSafe>& Context);
 
 private:
 
@@ -193,7 +193,7 @@ private:
 	TArray<TSharedPtr<FSessionLogMessage>> LogMessages;
 
 	/** Holds the message endpoint. */
-	FMessageEndpointPtr MessageEndpoint;
+	TSharedPtr<FMessageEndpoint, ESPMode::ThreadSafe> MessageEndpoint;
 
 	/** Holds a reference to the session that owns this instance. */
 	TWeakPtr<ISessionInfo> Owner;
diff --git a/Engine/Source/Runtime/SessionServices/Private/SessionManager.cpp b/Engine/Source/Runtime/SessionServices/Private/SessionManager.cpp
--- a/Engine/Source/Runtime/SessionServices/Private/SessionManager.cpp
+++ b/Engine/Source/Runtime/SessionServices/Private/SessionManager.cpp
@@ -5,7 +5,7 @@
 #include "Misc/CommandLine.h"
 #include "Containers/Ticker.h"
 #include "EngineServiceMessages.h"
-#include "Helpers/MessageEndpointBuilder.h"
+#include "MessageEndpointBuilder.h"
 #include "SessionServiceMessages.h"
 #include "SessionInfo.h"
 
@@ -246,7 +246,7 @@ void FSessionManager::SendPing()
 /* FSessionManager callbacks
  *****************************************************************************/
 
-void FSessionManager::HandleEnginePongMessage(const FEngineServicePong& Message, const IMessageContextRef& Context)
+void FSessionManager::HandleEnginePongMessage(const FEngineServicePong& Message, const TSharedRef<IMessageContext, ESPMode::ThreadSafe>& Context)
 {
 	if (!Message.SessionId.IsValid())
 	{
@@ -273,7 +273,7 @@ void FSessionManager::HandleLogReceived(const TSharedRef<ISessionInfo>& Session,
 }
 
 
-void FSessionManager::HandleSessionPongMessage(const FSessionServicePong& Message, const IMessageContextRef& Context)
+void FSessionManager::HandleSessionPongMessage(const FSessionServicePong& Message, const TSharedRef<IMessageContext, ESPMode::ThreadSafe>& Context)
 {
 	if (!Message.SessionId.IsValid())
 	{
diff --git a/Engine/Source/Runtime/SessionServices/Private/SessionManager.h b/Engine/Source/Runtime/SessionServices/Private/SessionManager.h
--- a/Engine/Source/Runtime/SessionServices/Private/SessionManager.h
+++ b/Engine/Source/Runtime/SessionServices/Private/SessionManager.h
@@ -6,7 +6,7 @@
 #include "Misc/Guid.h"
 #include "IMessageContext.h"
 #include "IMessageBus.h"
-#include "Helpers/MessageEndpoint.h"
+#include "MessageEndpoint.h"
 #include "ISessionInstanceInfo.h"
 #include "ISessionInfo.h"
 #include "SessionLogMessage.h"
@@ -123,7 +123,7 @@ private:
 	TWeakPtr<IMessageBus, ESPMode::ThreadSafe> MessageBusPtr;
 
 	/** Holds the messaging endpoint. */
-	FMessageEndpointPtr MessageEndpoint;
+	TSharedPtr<FMessageEndpoint, ESPMode::ThreadSafe> MessageEndpoint;
 
 	/** Holds the list of currently selected instances. */
 	TArray<TSharedPtr<ISessionInstanceInfo>> SelectedInstances;
diff --git a/Engine/Source/Runtime/SessionServices/Private/SessionService.cpp b/Engine/Source/Runtime/SessionServices/Private/SessionService.cpp
--- a/Engine/Source/Runtime/SessionServices/Private/SessionService.cpp
+++ b/Engine/Source/Runtime/SessionServices/Private/SessionService.cpp
@@ -1,9 +1,9 @@
 // Copyright 1998-2017 Epic Games, Inc. All Rights Reserved.
 
 #include "SessionService.h"
+#include "MessageEndpointBuilder.h"
 #include "Misc/ScopeLock.h"
 #include "Misc/App.h"
-#include "Helpers/MessageEndpointBuilder.h"
 #include "SessionServiceMessages.h"
 
 
diff --git a/Engine/Source/Runtime/SessionServices/Private/SessionService.h b/Engine/Source/Runtime/SessionServices/Private/SessionService.h
--- a/Engine/Source/Runtime/SessionServices/Private/SessionService.h
+++ b/Engine/Source/Runtime/SessionServices/Private/SessionService.h
@@ -5,13 +5,14 @@
 #include "CoreMinimal.h"
 #include "IMessageContext.h"
 #include "IMessageBus.h"
-#include "Helpers/MessageEndpoint.h"
 #include "ISessionService.h"
+#include "MessageEndpoint.h"
 
 struct FSessionServiceLogSubscribe;
 struct FSessionServiceLogUnsubscribe;
 struct FSessionServicePing;
 
+
 /**
  * Implements an application session service.
  */
diff --git a/Engine/Source/Runtime/SessionServices/SessionServices.Build.cs b/Engine/Source/Runtime/SessionServices/SessionServices.Build.cs
--- a/Engine/Source/Runtime/SessionServices/SessionServices.Build.cs
+++ b/Engine/Source/Runtime/SessionServices/SessionServices.Build.cs
@@ -7,25 +7,21 @@ public class SessionServices : ModuleRules
 	public SessionServices(ReadOnlyTargetRules Target) : base(Target)
 	{
 		PublicDependencyModuleNames.AddRange(
-			new string[]
-			{
+			new string[] {
 				"Core",
-			}
-		);
+			});
 
 		PrivateDependencyModuleNames.AddRange(
-			new string[]
-			{
+			new string[] {
 				"CoreUObject",
 				"EngineMessages",
 				"SessionMessages",
-			}
-		);
+			});
 
 		PrivateIncludePathModuleNames.AddRange(
 			new string[] {
 				"Messaging",
-			}
-		);
+				"MessagingCommon",
+			});
 	}
 }
diff --git a/Engine/Source/Runtime/Slate/Private/Framework/Application/HardwareCursor.cpp b/Engine/Source/Runtime/Slate/Private/Framework/Application/HardwareCursor.cpp
--- a/Engine/Source/Runtime/Slate/Private/Framework/Application/HardwareCursor.cpp
+++ b/Engine/Source/Runtime/Slate/Private/Framework/Application/HardwareCursor.cpp
@@ -1,22 +1,24 @@
 // Copyright 1998-2017 Epic Games, Inc. All Rights Reserved.
 
 #include "Framework/Application/HardwareCursor.h"
-#include "Logging/LogMacros.h"
-#include "Misc/Paths.h"
 #include "HAL/PlatformProcess.h"
+#include "IImageWrapper.h"
+#include "IImageWrapperModule.h"
+#include "Logging/LogMacros.h"
 #include "Misc/FileHelper.h"
+#include "Misc/Paths.h"
 #include "ModuleManager.h"
-#include "Interfaces/IImageWrapperModule.h"
 #include "SlateApplication.h"
 #include "HAL/PlatformApplicationMisc.h"
 
 #if PLATFORM_LINUX
-#include "SDL.h"
-#include "SDL_surface.h"
+	#include "SDL.h"
+	#include "SDL_surface.h"
 #endif
 
 DEFINE_LOG_CATEGORY_STATIC(LogHardwareCursor, Log, All);
 
+
 FHardwareCursor::FHardwareCursor(const FString& InPathToCursorWithoutExtension, FVector2D InHotSpot)
 #if PLATFORM_WINDOWS
 	: CursorHandle(NULL)
@@ -269,7 +271,7 @@ bool FHardwareCursor::LoadCursorFromPngs(const FString& InPathToCursorWithoutExt
 		}
 
 		IImageWrapperModule& ImageWrapperModule = FModuleManager::LoadModuleChecked<IImageWrapperModule>(FName("ImageWrapper"));
-		IImageWrapperPtr PngImageWrapper = ImageWrapperModule.CreateImageWrapper(EImageFormat::PNG);
+		TSharedPtr<IImageWrapper>PngImageWrapper = ImageWrapperModule.CreateImageWrapper(EImageFormat::PNG);
 
 		if ( PngImageWrapper.IsValid() && PngImageWrapper->SetCompressed(NearestCursor->FileData.GetData(), NearestCursor->FileData.Num()) )
 		{
diff --git a/Engine/Source/Runtime/Slate/Private/Framework/Text/Android/AndroidPlatformTextField.cpp b/Engine/Source/Runtime/Slate/Private/Framework/Text/Android/AndroidPlatformTextField.cpp
--- a/Engine/Source/Runtime/Slate/Private/Framework/Text/Android/AndroidPlatformTextField.cpp
+++ b/Engine/Source/Runtime/Slate/Private/Framework/Text/Android/AndroidPlatformTextField.cpp
@@ -1,6 +1,8 @@
 // Copyright 1998-2017 Epic Games, Inc. All Rights Reserved.
 
 #include "Framework/Text/Android/AndroidPlatformTextField.h"
+
+#include "HAL/IConsoleManager.h"
 #include "Widgets/Input/IVirtualKeyboardEntry.h"
 #include "Misc/ConfigCacheIni.h"
 #include "IConsoleManager.h"
diff --git a/Engine/Source/Runtime/Slate/Public/Widgets/Colors/SColorWheel.h b/Engine/Source/Runtime/Slate/Public/Widgets/Colors/SColorWheel.h
--- a/Engine/Source/Runtime/Slate/Public/Widgets/Colors/SColorWheel.h
+++ b/Engine/Source/Runtime/Slate/Public/Widgets/Colors/SColorWheel.h
@@ -53,7 +53,7 @@ public:
 
 public:
 
-	// SWidget overrides
+	//~ SWidget overrides
 
 	virtual FVector2D ComputeDesiredSize(float) const override;
 	virtual FReply OnMouseButtonDoubleClick(const FGeometry& InMyGeometry, const FPointerEvent& InMouseEvent) override;
diff --git a/Engine/Source/Runtime/Slate/Public/Widgets/Images/SThrobber.h b/Engine/Source/Runtime/Slate/Public/Widgets/Images/SThrobber.h
--- a/Engine/Source/Runtime/Slate/Public/Widgets/Images/SThrobber.h
+++ b/Engine/Source/Runtime/Slate/Public/Widgets/Images/SThrobber.h
@@ -14,13 +14,17 @@ class FPaintArgs;
 class FSlateWindowElementList;
 class SHorizontalBox;
 
-/** A Throbber widget that uses 5 zooming circles in a row.*/
-class SLATE_API SThrobber : public SCompoundWidget
-{
 
-static const int32 DefaultNumPieces = 3;
+/**
+ * A throbber widget that uses 5 zooming circles in a row.
+ */
+class SLATE_API SThrobber
+	: public SCompoundWidget
+{
+	static const int32 DefaultNumPieces = 3;
 
 public:
+
 	enum EAnimation
 	{
 		Vertical = 0x1 << 0,
@@ -57,13 +61,14 @@ public:
 	void SetAnimate(EAnimation InAnimate);
 
 private:
+
 	FVector2D GetPieceScale(int32 PieceIndex) const;
 	FLinearColor GetPieceColor(int32 PieceIndex) const;
 
 	/** Constructs the curves and widgets for the pieces which make up the throbber. */
 	void ConstructPieces();
 
-	/** Gets the brush used to draw each piece of the thobber. */
+	/** Gets the brush used to draw each piece of the throbber. */
 	const FSlateBrush* GetPieceBrush() const;
 
 private:
@@ -71,19 +76,25 @@ private:
 	FCurveSequence AnimCurves;
 	TArray<FCurveHandle, TInlineAllocator<DefaultNumPieces> > ThrobberCurve;
 
-	/** The horizontal box which contains the widgets for the trobber pieces. */
+	/** The horizontal box which contains the widgets for the throbber pieces. */
 	TSharedPtr<SHorizontalBox> HBox;
 
 	/** The image used to draw each piece of the throbber. */
 	const FSlateBrush* PieceImage;
+
 	/** The number of pieces to display. */
 	int NumPieces;
-	/** Flags for which apsects of the throbber to animate. */
+
+	/** Flags for which aspects of the throbber to animate. */
 	EAnimation Animate;
 };
 
-/** A throbber widget that orients images in a spinning circle. */
-class SLATE_API SCircularThrobber : public SLeafWidget
+
+/**
+ * A throbber widget that orients images in a spinning circle.
+ */
+class SLATE_API SCircularThrobber
+	: public SLeafWidget
 {
 public:
 	static const float MinimumPeriodValue;
@@ -126,20 +137,24 @@ public:
 
 private:
 
-	/** Constructs the sequence used to animate the trobber. */
+	/** Constructs the sequence used to animate the throbber. */
 	void ConstructSequence();
 
 private:
+
 	/** The sequence to drive the spinning animation */
 	FCurveSequence Sequence;
 	FCurveHandle Curve;
 
 	/** What each segment of the throbber looks like */
 	const FSlateBrush* PieceImage;
+
 	/** How many pieces there are */
 	int32 NumPieces;
+
 	/** The amount of time in seconds for a full circle */
 	float Period;
+
 	/** The radius of the circle */
 	float Radius;
 };
diff --git a/Engine/Source/Runtime/SlateRHIRenderer/Private/SlateRHIResourceManager.cpp b/Engine/Source/Runtime/SlateRHIRenderer/Private/SlateRHIResourceManager.cpp
--- a/Engine/Source/Runtime/SlateRHIRenderer/Private/SlateRHIResourceManager.cpp
+++ b/Engine/Source/Runtime/SlateRHIRenderer/Private/SlateRHIResourceManager.cpp
@@ -16,7 +16,8 @@
 #include "Engine/Engine.h"
 #include "Slate/SlateTextures.h"
 #include "SlateRHITextureAtlas.h"
-#include "Interfaces/IImageWrapperModule.h"
+#include "IImageWrapper.h"
+#include "IImageWrapperModule.h"
 #include "SlateNativeTextureResource.h"
 #include "SlateUTextureResource.h"
 #include "SlateMaterialResource.h"
@@ -390,7 +391,7 @@ bool FSlateRHIResourceManager::LoadTexture( const FName& TextureName, const FStr
 	if( FFileHelper::LoadFileToArray( RawFileData, *ResourcePath ) )
 	{
 		IImageWrapperModule& ImageWrapperModule = FModuleManager::LoadModuleChecked<IImageWrapperModule>( FName("ImageWrapper") );
-		IImageWrapperPtr ImageWrapper = ImageWrapperModule.CreateImageWrapper( EImageFormat::PNG );
+		TSharedPtr<IImageWrapper> ImageWrapper = ImageWrapperModule.CreateImageWrapper( EImageFormat::PNG );
 		if ( ImageWrapper.IsValid() && ImageWrapper->SetCompressed( RawFileData.GetData(), RawFileData.Num() ) )
 		{
 			Width = ImageWrapper->GetWidth();
diff --git a/Engine/Source/Runtime/Sockets/Private/BSDIPv6Sockets/SocketsBSDIPv6.cpp b/Engine/Source/Runtime/Sockets/Private/BSDIPv6Sockets/SocketsBSDIPv6.cpp
--- a/Engine/Source/Runtime/Sockets/Private/BSDIPv6Sockets/SocketsBSDIPv6.cpp
+++ b/Engine/Source/Runtime/Sockets/Private/BSDIPv6Sockets/SocketsBSDIPv6.cpp
@@ -49,7 +49,7 @@ EIPv6SocketInternalState::Return FSocketBSDIPv6::HasState(EIPv6SocketInternalSta
 	// convert WaitTime to a timeval
 	timeval Time;
 	Time.tv_sec = (int32)WaitTime.GetTotalSeconds();
-	Time.tv_usec = WaitTime.GetMilliseconds() * 1000 + WaitTime.GetMicroseconds();
+	Time.tv_usec = WaitTime.GetFractionMicro();
 
 	fd_set SocketSet;
 
diff --git a/Engine/Source/Runtime/Sockets/Private/BSDIPv6Sockets/SocketsBSDIPv6.h b/Engine/Source/Runtime/Sockets/Private/BSDIPv6Sockets/SocketsBSDIPv6.h
--- a/Engine/Source/Runtime/Sockets/Private/BSDIPv6Sockets/SocketsBSDIPv6.h
+++ b/Engine/Source/Runtime/Sockets/Private/BSDIPv6Sockets/SocketsBSDIPv6.h
@@ -133,7 +133,7 @@ public:
 protected:
 
 	// This is generally select(), but makes it easier for platforms without select to replace it
-	virtual EIPv6SocketInternalState::Return HasState(EIPv6SocketInternalState::Param State, FTimespan WaitTime=FTimespan(0));
+	virtual EIPv6SocketInternalState::Return HasState(EIPv6SocketInternalState::Param State, FTimespan WaitTime = FTimespan::Zero());
 
 	// Holds the BSD socket object.
 	SOCKET Socket;
diff --git a/Engine/Source/Runtime/Sockets/Private/BSDSockets/SocketsBSD.cpp b/Engine/Source/Runtime/Sockets/Private/BSDSockets/SocketsBSD.cpp
--- a/Engine/Source/Runtime/Sockets/Private/BSDSockets/SocketsBSD.cpp
+++ b/Engine/Source/Runtime/Sockets/Private/BSDSockets/SocketsBSD.cpp
@@ -438,7 +438,7 @@ ESocketBSDReturn FSocketBSD::HasState(ESocketBSDParam State, FTimespan WaitTime)
 	// convert WaitTime to a timeval
 	timeval Time;
 	Time.tv_sec = (int32)WaitTime.GetTotalSeconds();
-	Time.tv_usec = WaitTime.GetMilliseconds() * 1000 + WaitTime.GetMicroseconds();
+	Time.tv_usec = WaitTime.GetFractionMicro();
 
 	fd_set SocketSet;
 
diff --git a/Engine/Source/Runtime/Sockets/Private/BSDSockets/SocketsBSD.h b/Engine/Source/Runtime/Sockets/Private/BSDSockets/SocketsBSD.h
--- a/Engine/Source/Runtime/Sockets/Private/BSDSockets/SocketsBSD.h
+++ b/Engine/Source/Runtime/Sockets/Private/BSDSockets/SocketsBSD.h
@@ -114,7 +114,7 @@ public:
 protected:
 
 	/** This is generally select(), but makes it easier for platforms without select to replace it. */
-	virtual ESocketBSDReturn HasState(ESocketBSDParam State, FTimespan WaitTime = FTimespan(0));
+	virtual ESocketBSDReturn HasState(ESocketBSDParam State, FTimespan WaitTime = FTimespan::Zero());
 
 	/** Updates this socket's time of last activity. */
 	void UpdateActivity()
diff --git a/Engine/Source/Runtime/Sockets/Private/HTML5/Simulator/SocketX.cpp b/Engine/Source/Runtime/Sockets/Private/HTML5/Simulator/SocketX.cpp
--- a/Engine/Source/Runtime/Sockets/Private/HTML5/Simulator/SocketX.cpp
+++ b/Engine/Source/Runtime/Sockets/Private/HTML5/Simulator/SocketX.cpp
@@ -78,14 +78,14 @@ bool FSocketX::Wait( ESocketWaitConditions::Type Condition, FTimespan WaitTime )
 {
 	if ((Condition == ESocketWaitConditions::WaitForRead) )
 	{
-		return Pimpl->WaitForRead( WaitTime.GetMilliseconds() );
+		return Pimpl->WaitForRead((int)WaitTime.GetTotalMilliseconds());
 	}
 	if ( Condition == ESocketWaitConditions::WaitForWrite)
 	{
-		return Pimpl->WaitForWrite( WaitTime.GetMilliseconds() );
+		return Pimpl->WaitForWrite((int)WaitTime.GetTotalMilliseconds());
 	}
 		
-	return Pimpl->WaitForReadWrite( WaitTime.GetMilliseconds() );
+	return Pimpl->WaitForReadWrite((int)WaitTime.GetTotalMilliseconds());
 }
 
 ESocketConnectionState FSocketX::GetConnectionState()
diff --git a/Engine/Source/Runtime/UMG/Private/Blueprint/AsyncTaskDownloadImage.cpp b/Engine/Source/Runtime/UMG/Private/Blueprint/AsyncTaskDownloadImage.cpp
--- a/Engine/Source/Runtime/UMG/Private/Blueprint/AsyncTaskDownloadImage.cpp
+++ b/Engine/Source/Runtime/UMG/Private/Blueprint/AsyncTaskDownloadImage.cpp
@@ -4,7 +4,8 @@
 #include "Modules/ModuleManager.h"
 #include "Engine/Texture2D.h"
 #include "Engine/Texture2DDynamic.h"
-#include "Interfaces/IImageWrapperModule.h"
+#include "IImageWrapper.h"
+#include "IImageWrapperModule.h"
 #include "Interfaces/IHttpResponse.h"
 #include "HttpModule.h"
 
@@ -89,7 +90,7 @@ void UAsyncTaskDownloadImage::HandleImageRequest(FHttpRequestPtr HttpRequest, FH
 	if ( bSucceeded && HttpResponse.IsValid() && HttpResponse->GetContentLength() > 0 )
 	{
 		IImageWrapperModule& ImageWrapperModule = FModuleManager::LoadModuleChecked<IImageWrapperModule>(FName("ImageWrapper"));
-		IImageWrapperPtr ImageWrappers[3] =
+		TSharedPtr<IImageWrapper> ImageWrappers[3] =
 		{
 			ImageWrapperModule.CreateImageWrapper(EImageFormat::PNG),
 			ImageWrapperModule.CreateImageWrapper(EImageFormat::JPEG),
diff --git a/Engine/Source/Runtime/UtilityShaders/Private/MediaShaders.cpp b/Engine/Source/Runtime/UtilityShaders/Private/MediaShaders.cpp
--- a/Engine/Source/Runtime/UtilityShaders/Private/MediaShaders.cpp
+++ b/Engine/Source/Runtime/UtilityShaders/Private/MediaShaders.cpp
@@ -89,6 +89,7 @@ void FAYUVConvertPS::SetParameters(FRHICommandList& CommandList, TRefCountPtr<FR
  *****************************************************************************/
 
 BEGIN_UNIFORM_BUFFER_STRUCT(FBMPConvertUB, )
+DECLARE_UNIFORM_BUFFER_STRUCT_MEMBER(uint32, SrgbToLinear)
 DECLARE_UNIFORM_BUFFER_STRUCT_MEMBER(FVector2D, UVScale)
 DECLARE_UNIFORM_BUFFER_STRUCT_MEMBER_TEXTURE(Texture2D, Texture)
 DECLARE_UNIFORM_BUFFER_STRUCT_MEMBER_SAMPLER(SamplerState, Sampler)
@@ -98,10 +99,11 @@ IMPLEMENT_UNIFORM_BUFFER_STRUCT(FBMPConvertUB, TEXT("BMPConvertUB"));
 IMPLEMENT_SHADER_TYPE(, FBMPConvertPS, TEXT("/Engine/Private/MediaShaders.usf"), TEXT("BMPConvertPS"), SF_Pixel);
 
 
-void FBMPConvertPS::SetParameters(FRHICommandList& CommandList, TRefCountPtr<FRHITexture2D> BMPTexture, const FIntPoint& OutputDimensions)
+void FBMPConvertPS::SetParameters(FRHICommandList& CommandList, TRefCountPtr<FRHITexture2D> BMPTexture, const FIntPoint& OutputDimensions, bool SrgbToLinear)
 {
 	FBMPConvertUB UB;
 	{
+		UB.SrgbToLinear = SrgbToLinear;
 		UB.UVScale = FVector2D((float)OutputDimensions.X / (float)BMPTexture->GetSizeX(), (float)OutputDimensions.Y / (float)BMPTexture->GetSizeY());
 		UB.Texture = BMPTexture;
 		UB.Sampler = TStaticSamplerState<SF_Bilinear>::GetRHI();
@@ -187,7 +189,6 @@ void FNV21ConvertPS::SetParameters(FRHICommandList& CommandList, TRefCountPtr<FR
 
 BEGIN_UNIFORM_BUFFER_STRUCT(FRGBConvertUB, )
 DECLARE_UNIFORM_BUFFER_STRUCT_MEMBER(FVector2D, UVScale)
-DECLARE_UNIFORM_BUFFER_STRUCT_MEMBER(uint32, SrgbToLinear)
 DECLARE_UNIFORM_BUFFER_STRUCT_MEMBER_TEXTURE(Texture2D, Texture)
 DECLARE_UNIFORM_BUFFER_STRUCT_MEMBER_SAMPLER(SamplerState, Sampler)
 END_UNIFORM_BUFFER_STRUCT(FRGBConvertUB)
@@ -196,7 +197,7 @@ IMPLEMENT_UNIFORM_BUFFER_STRUCT(FRGBConvertUB, TEXT("RGBConvertUB"));
 IMPLEMENT_SHADER_TYPE(, FRGBConvertPS, TEXT("/Engine/Private/MediaShaders.usf"), TEXT("RGBConvertPS"), SF_Pixel);
 
 
-void FRGBConvertPS::SetParameters(FRHICommandList& CommandList, TRefCountPtr<FRHITexture2D> RGBTexture, const FIntPoint& OutputDimensions, bool SrgbToLinear)
+void FRGBConvertPS::SetParameters(FRHICommandList& CommandList, TRefCountPtr<FRHITexture2D> RGBTexture, const FIntPoint& OutputDimensions)
 {
 	FRGBConvertUB UB;
 	{
diff --git a/Engine/Source/Runtime/UtilityShaders/Public/MediaShaders.h b/Engine/Source/Runtime/UtilityShaders/Public/MediaShaders.h
--- a/Engine/Source/Runtime/UtilityShaders/Public/MediaShaders.h
+++ b/Engine/Source/Runtime/UtilityShaders/Public/MediaShaders.h
@@ -171,7 +171,7 @@ public:
 		return bShaderHasOutdatedParameters;
 	}
 
-	UTILITYSHADERS_API void SetParameters(FRHICommandList& RHICmdList, TRefCountPtr<FRHITexture2D> BMPTexture, const FIntPoint& OutputDimensions);
+	UTILITYSHADERS_API void SetParameters(FRHICommandList& RHICmdList, TRefCountPtr<FRHITexture2D> BMPTexture, const FIntPoint& OutputDimensions, bool SrgbToLinear);
 };
 
 
@@ -275,7 +275,7 @@ public:
 		return bShaderHasOutdatedParameters;
 	}
 
-	UTILITYSHADERS_API void SetParameters(FRHICommandList& RHICmdList, TRefCountPtr<FRHITexture2D> RGBTexture, const FIntPoint& OutputDimensions, bool SrgbToLinear);
+	UTILITYSHADERS_API void SetParameters(FRHICommandList& RHICmdList, TRefCountPtr<FRHITexture2D> RGBTexture, const FIntPoint& OutputDimensions);
 };
 
 
