HTML5 - Unreal File Server - css response header fixes

#rnx

diff --git a/Engine/Source/Runtime/NetworkFileSystem/Private/NetworkFileServerHttp.cpp b/Engine/Source/Runtime/NetworkFileSystem/Private/NetworkFileServerHttp.cpp
--- a/Engine/Source/Runtime/NetworkFileSystem/Private/NetworkFileServerHttp.cpp
+++ b/Engine/Source/Runtime/NetworkFileSystem/Private/NetworkFileServerHttp.cpp
@@ -370,6 +370,10 @@ int FNetworkFileServerHttp::CallBack_HTTP(
 				{
 					FCStringWide::Strcpy(Mime,TEXT("application/javascript;charset=UTF-8"));
 				}
+				else if ( FilePath.Contains(".css"))
+				{
+					FCStringWide::Strcpy(Mime,TEXT("text/css;charset=UTF-8"));
+				}
 				else
 				{
 					FCStringWide::Strcpy(Mime,TEXT("text/html;charset=UTF-8"));
