HTML5 - FrontEnd COTF - seems proxy port is no longer used/needed

#rnx

diff --git a/Engine/Source/Programs/AutomationTool/HTML5/HTML5Platform.Automation.cs b/Engine/Source/Programs/AutomationTool/HTML5/HTML5Platform.Automation.cs
--- a/Engine/Source/Programs/AutomationTool/HTML5/HTML5Platform.Automation.cs
+++ b/Engine/Source/Programs/AutomationTool/HTML5/HTML5Platform.Automation.cs
@@ -587,19 +587,20 @@ public class HTML5Platform : Platform
 		string WorkingDirectory = Path.GetDirectoryName(ClientApp);
 		string url = Path.GetFileName(ClientApp) +".html";
 
-		// WARNING: splitting the following situation
-		// if cookonthefly is used: tell the browser to use the PROXY at DEFAULT_HTTP_FILE_SERVING_PORT
-		// leave the normal HTML5LaunchHelper port (ServerPort) alone -- otherwise it will collide with the PROXY port
+// UE-64628: seems proxy port is no longer used anymore -- leaving details here for future reference...
+//		// WARNING: splitting the following situation
+//		// if cookonthefly is used: tell the browser to use the PROXY at DEFAULT_HTTP_FILE_SERVING_PORT
+//		// leave the normal HTML5LaunchHelper port (ServerPort) alone -- otherwise it will collide with the PROXY port
 		if (ClientCmdLine.Contains("filehostip"))
 		{
 			url += "?cookonthefly=true";
-			Int32 ProxyPort = 41898; // DEFAULT_HTTP_FILE_SERVING_PORT
-			url = String.Format("http://localhost:{0}/{1}", ProxyPort, url);
+//			Int32 ProxyPort = 41898; // DEFAULT_HTTP_FILE_SERVING_PORT
+//			url = String.Format("http://localhost:{0}/{1}", ProxyPort, url);
 		}
-		else
-		{
+//		else
+//		{
 			url = String.Format("http://localhost:{0}/{1}", ServerPort, url);
-		}
+//		}
 
 		// Check HTML5LaunchHelper source for command line args
 
diff --git a/Engine/Source/Runtime/NetworkFileSystem/Private/NetworkFileServerHttp.cpp b/Engine/Source/Runtime/NetworkFileSystem/Private/NetworkFileServerHttp.cpp
--- a/Engine/Source/Runtime/NetworkFileSystem/Private/NetworkFileServerHttp.cpp
+++ b/Engine/Source/Runtime/NetworkFileSystem/Private/NetworkFileServerHttp.cpp
@@ -347,6 +347,7 @@ int FNetworkFileServerHttp::CallBack_HTTP(
 					Buffer,
 					TEXT("HTTP/1.0 200 OK\x0d\x0a")
 					TEXT("Server: Unreal File Server\x0d\x0a")
+					TEXT("Access-Control-Allow-Origin: *\x0d\x0a")
 					TEXT("Connection: close\x0d\x0a")
 					TEXT("Content-Type: text/html; charset=utf-8\x0d\x0a")
 					TEXT("Content-Length: %u\x0d\x0a\x0d\x0a%s"),
@@ -421,6 +422,7 @@ int FNetworkFileServerHttp::CallBack_HTTP(
 					Length = FCString::Sprintf(Header,
 						TEXT("HTTP/1.1 200 OK\x0d\x0a")
 						TEXT("Server: Unreal File Server\x0d\x0a")
+						TEXT("Access-Control-Allow-Origin: *\x0d\x0a")
 						TEXT("Connection: close\x0d\x0a")
 						TEXT("Content-Type: %s \x0d\x0a")
 						TEXT("Content-Encoding: gzip\x0d\x0a")
@@ -432,6 +434,7 @@ int FNetworkFileServerHttp::CallBack_HTTP(
 					Length = FCString::Sprintf(Header,
 						TEXT("HTTP/1.1 200 OK\x0d\x0a")
 						TEXT("Server: Unreal File Server\x0d\x0a")
+						TEXT("Access-Control-Allow-Origin: *\x0d\x0a")
 						TEXT("Connection: close\x0d\x0a")
 						TEXT("Content-Type: %s \x0d\x0a")
 						TEXT("Content-Length: %u\x0d\x0a\x0d\x0a"),
@@ -477,6 +480,7 @@ int FNetworkFileServerHttp::CallBack_HTTP(
 				(ANSICHAR*)Header,
 				"HTTP/1.1 200 OK\x0d\x0a"
 				"Server: Unreal File Server\x0d\x0a"
+				"Access-Control-Allow-Origin: *\x0d\x0a"
 				"Connection: close\x0d\x0a"
 				"Content-Type: application/octet-stream \x0d\x0a"
 				"Content-Length: %u\x0d\x0a\x0d\x0a",
