COTF: Don't send uncooked packages when cook fails

#rnx

diff --git a/Engine/Source/Runtime/NetworkFileSystem/Private/NetworkFileServerConnection.cpp b/Engine/Source/Runtime/NetworkFileSystem/Private/NetworkFileServerConnection.cpp
--- a/Engine/Source/Runtime/NetworkFileSystem/Private/NetworkFileServerConnection.cpp
+++ b/Engine/Source/Runtime/NetworkFileSystem/Private/NetworkFileServerConnection.cpp
@@ -457,6 +457,10 @@ void FNetworkFileServerClientConnection::ProcessOpenFile( FArchive& In, FArchive
 	TArray<FString> NewUnsolictedFiles;
 	NetworkFileDelegates->FileRequestDelegate.ExecuteIfBound(Filename, ConnectedPlatformName, NewUnsolictedFiles);
 
+	// Disable access to outside the sandbox to prevent sending uncooked packages to the client
+	const bool bIsCookable = FPackageName::IsPackageExtension(*FPaths::GetExtension(Filename, true));
+	Sandbox->SetSandboxOnly(bIsCookable && !bIsWriting);
+
 	FDateTime ServerTimeStamp = Sandbox->GetTimeStamp(*Filename);
 	int64 ServerFileSize = 0;
 	IFileHandle* File = bIsWriting ? Sandbox->OpenWrite(*Filename, bAppend, bAllowRead) : Sandbox->OpenRead(*Filename);
@@ -470,6 +474,8 @@ void FNetworkFileServerClientConnection::ProcessOpenFile( FArchive& In, FArchive
 		ServerFileSize = File->Size();
 	}
 
+	Sandbox->SetSandboxOnly(false);
+
 	uint64 HandleId = ++LastHandleId;
 	OpenFiles.Add( HandleId, File );
 	
diff --git a/Engine/Source/Runtime/SandboxFile/Private/IPlatformFileSandboxWrapper.cpp b/Engine/Source/Runtime/SandboxFile/Private/IPlatformFileSandboxWrapper.cpp
--- a/Engine/Source/Runtime/SandboxFile/Private/IPlatformFileSandboxWrapper.cpp
+++ b/Engine/Source/Runtime/SandboxFile/Private/IPlatformFileSandboxWrapper.cpp
@@ -35,6 +35,7 @@ FSandboxPlatformFile::FSandboxPlatformFile(bool bInEntireEngineWillUseThisSandbo
 	: LowerLevel(nullptr)
 	, bEntireEngineWillUseThisSandbox(bInEntireEngineWillUseThisSandbox)
 	, bSandboxEnabled(true)
+	, bSandboxOnly(false)
 {
 }
 
@@ -447,7 +448,7 @@ bool FSandboxPlatformFile::OkForInnerAccess(const TCHAR* InFilenameOrDirectoryNa
 			}
 		}
 	}
-	return true;
+	return !bSandboxOnly;
 }
 
 void FSandboxPlatformFile::SetSandboxEnabled(bool bInEnabled)
@@ -487,6 +488,11 @@ void FSandboxPlatformFile::AddExclusion(const TCHAR* Wildcard, bool bIsDirectory
 	}
 }
 
+void FSandboxPlatformFile::SetSandboxOnly(bool bInSandboxOnly)
+{
+	bSandboxOnly = bInSandboxOnly;
+}
+
 // IPlatformFile Interface
 
 bool FSandboxPlatformFile::FileExists(const TCHAR* Filename)
diff --git a/Engine/Source/Runtime/SandboxFile/Public/IPlatformFileSandboxWrapper.h b/Engine/Source/Runtime/SandboxFile/Public/IPlatformFileSandboxWrapper.h
--- a/Engine/Source/Runtime/SandboxFile/Public/IPlatformFileSandboxWrapper.h
+++ b/Engine/Source/Runtime/SandboxFile/Public/IPlatformFileSandboxWrapper.h
@@ -44,6 +44,12 @@ class SANDBOXFILE_API FSandboxPlatformFile : public IPlatformFile
 	 */
 	bool				bSandboxEnabled;
 
+	/**
+	 *	Whether access is restricted to the sandbox or not.
+	 *	Defaults to false.
+	 */
+	bool				bSandboxOnly;
+
 	/**
 	 * Clears the contents of the specified folder
 	 *
@@ -163,6 +169,9 @@ public:
 	*/
 	void AddExclusion(const TCHAR* Wildcard, bool bIsDirectory = false);
 
+	/** Whether access is restricted the the sandbox or not. */
+	void SetSandboxOnly(bool bInSandboxOnly);
+
 	// IPlatformFile Interface
 
 	virtual bool					FileExists(const TCHAR* Filename) override;
