Don't send uncooked packages when requesting files with SyncFile

#rnx

diff --git a/Engine/Source/Runtime/NetworkFileSystem/Private/NetworkFileServerConnection.cpp b/Engine/Source/Runtime/NetworkFileSystem/Private/NetworkFileServerConnection.cpp
--- a/Engine/Source/Runtime/NetworkFileSystem/Private/NetworkFileServerConnection.cpp
+++ b/Engine/Source/Runtime/NetworkFileSystem/Private/NetworkFileServerConnection.cpp
@@ -1213,6 +1213,10 @@ bool FNetworkFileServerClientConnection::PackageFile( FString& Filename, FString
 	// get file timestamp and send it to client
 	FDateTime ServerTimeStamp = Sandbox->GetTimeStamp(*Filename);
 
+	// Disable access to outside the sandbox to prevent sending uncooked packages to the client
+	const bool bIsCookable = FPackageName::IsPackageExtension(*FPaths::GetExtension(Filename, true));
+	Sandbox->SetSandboxOnly(bIsCookable);
+
 	FString AbsHostFile = Sandbox->ConvertToAbsolutePathForExternalAppForRead(*Filename);
 	if (ConnectedTargetPlatform != nullptr && ConnectedTargetPlatform->CopyFileToTarget(ConnectedIPAddress, AbsHostFile, TargetFilename, ConnectedTargetCustomData))
 	{
@@ -1260,6 +1264,8 @@ bool FNetworkFileServerClientConnection::PackageFile( FString& Filename, FString
 		UE_LOG(LogFileServer, Display, TEXT("Read %s, %d bytes"), *Filename, Contents.Num());
 	}
 
+	Sandbox->SetSandboxOnly(false);
+
 	Out << Filename;
 	Out << ServerTimeStamp;
 	uint64 FileSize = Contents.Num();
