Added option to network file server to prevent package assets being sent from outside the cooked sandbox folder

#rnx

diff --git a/Engine/Source/Editor/UnrealEd/Private/CookOnTheFlyServer.cpp b/Engine/Source/Editor/UnrealEd/Private/CookOnTheFlyServer.cpp
--- a/Engine/Source/Editor/UnrealEd/Private/CookOnTheFlyServer.cpp
+++ b/Engine/Source/Editor/UnrealEd/Private/CookOnTheFlyServer.cpp
@@ -555,19 +555,19 @@ bool UCookOnTheFlyServer::StartNetworkFileServer(const bool BindAnyPort, const T
 	FSandboxPathDelegate SandboxPathDelegate(FSandboxPathDelegate::CreateUObject(this, &UCookOnTheFlyServer::HandleNetworkGetSandboxPath));
 	FInitialPrecookedListDelegate InitialPrecookedListDelegate(FInitialPrecookedListDelegate::CreateUObject(this, &UCookOnTheFlyServer::HandleNetworkGetPrecookedList));
 
-
-	FNetworkFileDelegateContainer NetworkFileDelegateContainer;
-	NetworkFileDelegateContainer.NewConnectionDelegate = NewConnectionDelegate;
-	NetworkFileDelegateContainer.InitialPrecookedListDelegate = InitialPrecookedListDelegate;
-	NetworkFileDelegateContainer.FileRequestDelegate = FileRequestDelegate;
-	NetworkFileDelegateContainer.RecompileShadersDelegate = RecompileShadersDelegate;
-	NetworkFileDelegateContainer.SandboxPathOverrideDelegate = SandboxPathDelegate;
-	
-	NetworkFileDelegateContainer.OnFileModifiedCallback = &FileModifiedDelegate;
-
+	FNetworkFileServerOptions FileServerOptions;
+	FileServerOptions.Protocol = NFSP_Tcp;
+	FileServerOptions.Port = BindAnyPort ? 0 : -1;
+	FileServerOptions.Delegates.NewConnectionDelegate = NewConnectionDelegate;
+	FileServerOptions.Delegates.InitialPrecookedListDelegate = InitialPrecookedListDelegate;
+	FileServerOptions.Delegates.FileRequestDelegate = FileRequestDelegate;
+	FileServerOptions.Delegates.RecompileShadersDelegate = RecompileShadersDelegate;
+	FileServerOptions.Delegates.SandboxPathOverrideDelegate = SandboxPathDelegate;
+	FileServerOptions.Delegates.OnFileModifiedCallback = &FileModifiedDelegate;
+	FileServerOptions.bRestrictPackageAssetsToSandbox = true; // prevents sending uncooked packages if the package fails to cook
 
 	INetworkFileServer *TcpFileServer = FModuleManager::LoadModuleChecked<INetworkFileSystemModule>("NetworkFileSystem")
-		.CreateNetworkFileServer(true, BindAnyPort ? 0 : -1, NetworkFileDelegateContainer, ENetworkFileServerProtocol::NFSP_Tcp);
+		.CreateNetworkFileServer(MoveTemp(FileServerOptions), /* bLoadTargetPlatforms */ true);
 	if ( TcpFileServer )
 	{
 		NetworkFileServers.Add(TcpFileServer);
@@ -8592,4 +8592,4 @@ uint32 UCookOnTheFlyServer::FullLoadAndSave(uint32& CookedPackageCount)
 	return Result;
 }
 
-#undef LOCTEXT_NAMESPACE
\ No newline at end of file
+#undef LOCTEXT_NAMESPACE
diff --git a/Engine/Source/Runtime/NetworkFileSystem/Private/NetworkFileServer.cpp b/Engine/Source/Runtime/NetworkFileSystem/Private/NetworkFileServer.cpp
--- a/Engine/Source/Runtime/NetworkFileSystem/Private/NetworkFileServer.cpp
+++ b/Engine/Source/Runtime/NetworkFileSystem/Private/NetworkFileServer.cpp
@@ -16,9 +16,9 @@ class FNetworkFileServerClientConnectionThreaded
 {
 public:
 
-	FNetworkFileServerClientConnectionThreaded(FSocket* InSocket, const FNetworkFileDelegateContainer* NetworkFileDelegates, const TArray<ITargetPlatform*>& InActiveTargetPlatforms )
-		:  FNetworkFileServerClientConnection( NetworkFileDelegates, InActiveTargetPlatforms)
-		  ,Socket(InSocket)
+	FNetworkFileServerClientConnectionThreaded(FSocket* InSocket, const FNetworkFileServerOptions& Options)
+		: FNetworkFileServerClientConnection(Options)
+		, Socket(InSocket)
 	{
 		Running.Set(true);
 		StopRequested.Reset();
@@ -112,21 +112,18 @@ private:
 /* FNetworkFileServer constructors
  *****************************************************************************/
 
-FNetworkFileServer::FNetworkFileServer( int32 InPort, FNetworkFileDelegateContainer InNetworkFileDelegateContainer,
-	const TArray<ITargetPlatform*>& InActiveTargetPlatforms )
-	:ActiveTargetPlatforms(InActiveTargetPlatforms)
+FNetworkFileServer::FNetworkFileServer(FNetworkFileServerOptions InFileServerOptions)
+	: FileServerOptions(MoveTemp(InFileServerOptions))
 {
-	if(InPort <0)
+	if (FileServerOptions.Port < 0 )
 	{
-		InPort = DEFAULT_TCP_FILE_SERVING_PORT;
+		FileServerOptions.Port = DEFAULT_TCP_FILE_SERVING_PORT;
 	}
 
 	Running.Set(false);
 	StopRequested.Set(false);
 	UE_LOG(LogFileServer, Warning, TEXT("Unreal Network File Server starting up..."));
 
-	NetworkFileDelegates = InNetworkFileDelegateContainer;
-
 	// make sure sockets are going
 	ISocketSubsystem* SocketSubsystem = ISocketSubsystem::Get();
 
@@ -138,7 +135,7 @@ FNetworkFileServer::FNetworkFileServer( int32 InPort, FNetworkFileDelegateContai
 	{
 		// listen on any IP address
 		ListenAddr = SocketSubsystem->GetLocalBindAddr(*GLog);
-		ListenAddr->SetPort(InPort);
+		ListenAddr->SetPort(FileServerOptions.Port);
 
 		// create a server TCP socket
 		Socket = SocketSubsystem->CreateSocket(NAME_Stream, TEXT("FNetworkFileServer tcp-listen"), ListenAddr->GetProtocolType());
@@ -165,7 +162,7 @@ FNetworkFileServer::FNetworkFileServer( int32 InPort, FNetworkFileDelegateContai
 			{
 				// set the port on the listen address to be the same as the port on the socket
 				int32 port = Socket->GetPortNo();
-				check((InPort == 0 && port != 0) || port == InPort);
+				check((FileServerOptions.Port == 0 && port != 0) || port == FileServerOptions.Port);
 				ListenAddr->SetPort(port);
 
 				// now create a thread to accept connections
@@ -176,7 +173,7 @@ FNetworkFileServer::FNetworkFileServer( int32 InPort, FNetworkFileDelegateContai
 		}
 	}
 
- }
+}
 
 FNetworkFileServer::~FNetworkFileServer()
 {
@@ -249,7 +246,7 @@ uint32 FNetworkFileServer::Run( )
 						}
 					}
 
-					FNetworkFileServerClientConnectionThreaded* Connection = new FNetworkFileServerClientConnectionThreaded(ClientSocket, &NetworkFileDelegates, ActiveTargetPlatforms);
+					FNetworkFileServerClientConnectionThreaded* Connection = new FNetworkFileServerClientConnectionThreaded(ClientSocket, FileServerOptions);
 					Connections.Add(Connection);
 					UE_LOG(LogFileServer, Display, TEXT("Client %s connected."), *Connection->GetDescription());
 				}
diff --git a/Engine/Source/Runtime/NetworkFileSystem/Private/NetworkFileServer.h b/Engine/Source/Runtime/NetworkFileSystem/Private/NetworkFileServer.h
--- a/Engine/Source/Runtime/NetworkFileSystem/Private/NetworkFileServer.h
+++ b/Engine/Source/Runtime/NetworkFileSystem/Private/NetworkFileServer.h
@@ -25,10 +25,9 @@ public:
 	/**
 	 * Creates and initializes a new instance.
 	 *
-	 * @param InPort The port number to bind to (0 = any available port).
-	 * @param InFileRequestDelegate 
+	 * @param InFileServerOptions Network file server options
 	 */
-	FNetworkFileServer( int32 InPort, FNetworkFileDelegateContainer InNetworkFileDelegateContainer, const TArray<ITargetPlatform*>& InActiveTargetPlatforms );
+	FNetworkFileServer(FNetworkFileServerOptions InFileServerOptions);
 
 	/**
 	 * Destructor.
@@ -64,6 +63,9 @@ public:
 	virtual void Shutdown() override;
 private:
 
+	// File server options
+	FNetworkFileServerOptions FileServerOptions;
+
 	// Holds the server (listening) socket.
 	FSocket* Socket;
 
@@ -79,13 +81,6 @@ private:
 	// Is the Listner thread up and running. 
 	FThreadSafeCounter Running;
 
-public:
-
-	FNetworkFileDelegateContainer NetworkFileDelegates;
-
-	// cached copy of the active target platforms (if any)
-	const TArray<ITargetPlatform*> ActiveTargetPlatforms;
-
 	// Holds the address that the server is bound to.
 	TSharedPtr<FInternetAddr> ListenAddr;
 };
diff --git a/Engine/Source/Runtime/NetworkFileSystem/Private/NetworkFileServerConnection.cpp b/Engine/Source/Runtime/NetworkFileSystem/Private/NetworkFileServerConnection.cpp
--- a/Engine/Source/Runtime/NetworkFileSystem/Private/NetworkFileServerConnection.cpp
+++ b/Engine/Source/Runtime/NetworkFileSystem/Private/NetworkFileServerConnection.cpp
@@ -49,11 +49,12 @@ static FString MakeAbsoluteNormalizedDir(const FString& InPath)
 /* FNetworkFileServerClientConnection structors
  *****************************************************************************/
 
-FNetworkFileServerClientConnection::FNetworkFileServerClientConnection( const FNetworkFileDelegateContainer* InNetworkFileDelegates, const TArray<ITargetPlatform*>& InActiveTargetPlatforms )
+FNetworkFileServerClientConnection::FNetworkFileServerClientConnection(const FNetworkFileServerOptions& Options)
 	: LastHandleId(0)
 	, Sandbox(NULL)
-	, NetworkFileDelegates(InNetworkFileDelegates)
-	, ActiveTargetPlatforms(InActiveTargetPlatforms)
+	, NetworkFileDelegates(&Options.Delegates)
+	, ActiveTargetPlatforms(Options.TargetPlatforms)
+	, bRestrictPackageAssetsToSandbox(Options.bRestrictPackageAssetsToSandbox)
 {	
 	//stats
 	FileRequestDelegateTime = 0.0;
@@ -458,8 +459,11 @@ void FNetworkFileServerClientConnection::ProcessOpenFile( FArchive& In, FArchive
 	NetworkFileDelegates->FileRequestDelegate.ExecuteIfBound(Filename, ConnectedPlatformName, NewUnsolictedFiles);
 
 	// Disable access to outside the sandbox to prevent sending uncooked packages to the client
-	const bool bIsCookable = FPackageName::IsPackageExtension(*FPaths::GetExtension(Filename, true));
-	Sandbox->SetSandboxOnly(bIsCookable && !bIsWriting);
+	if (bRestrictPackageAssetsToSandbox && !bIsWriting)
+	{
+		const bool bIsPackageAsset = FPackageName::IsPackageExtension(*FPaths::GetExtension(Filename, true));
+		Sandbox->SetSandboxOnly(bIsPackageAsset);
+	}
 
 	FDateTime ServerTimeStamp = Sandbox->GetTimeStamp(*Filename);
 	int64 ServerFileSize = 0;
@@ -1214,8 +1218,11 @@ bool FNetworkFileServerClientConnection::PackageFile( FString& Filename, FString
 	FDateTime ServerTimeStamp = Sandbox->GetTimeStamp(*Filename);
 
 	// Disable access to outside the sandbox to prevent sending uncooked packages to the client
-	const bool bIsCookable = FPackageName::IsPackageExtension(*FPaths::GetExtension(Filename, true));
-	Sandbox->SetSandboxOnly(bIsCookable);
+	if (bRestrictPackageAssetsToSandbox)
+	{
+		const bool bIsPackageAsset = FPackageName::IsPackageExtension(*FPaths::GetExtension(Filename, true));
+		Sandbox->SetSandboxOnly(bIsPackageAsset);
+	}
 
 	FString AbsHostFile = Sandbox->ConvertToAbsolutePathForExternalAppForRead(*Filename);
 	if (ConnectedTargetPlatform != nullptr && ConnectedTargetPlatform->CopyFileToTarget(ConnectedIPAddress, AbsHostFile, TargetFilename, ConnectedTargetCustomData))
diff --git a/Engine/Source/Runtime/NetworkFileSystem/Private/NetworkFileServerConnection.h b/Engine/Source/Runtime/NetworkFileSystem/Private/NetworkFileServerConnection.h
--- a/Engine/Source/Runtime/NetworkFileSystem/Private/NetworkFileServerConnection.h
+++ b/Engine/Source/Runtime/NetworkFileSystem/Private/NetworkFileServerConnection.h
@@ -23,7 +23,7 @@ public:
 	 * @param InSocket - The client socket to use.
 	 * @param NetworkFileDelegates- delegates the client calls when events from the client happen
 	 */
-	FNetworkFileServerClientConnection( const FNetworkFileDelegateContainer* NetworkFileDelegates , const TArray<ITargetPlatform*>& InActiveTargetPlatforms );
+	FNetworkFileServerClientConnection(const FNetworkFileServerOptions& InFileServerOptions);
 
 	/**
 	 * Destructor.
@@ -285,6 +285,8 @@ private:
 	// custom key-value pair data for the curently connected target platform
 	TMap<FString,FString> ConnectedTargetCustomData;
 
+	// whether to restrict sending package assets from outside the sandbox
+	bool bRestrictPackageAssetsToSandbox;
 
 	//////////////////////////////////////////////////////////////////////////
 	//stats
diff --git a/Engine/Source/Runtime/NetworkFileSystem/Private/NetworkFileServerHttp.cpp b/Engine/Source/Runtime/NetworkFileSystem/Private/NetworkFileServerHttp.cpp
--- a/Engine/Source/Runtime/NetworkFileSystem/Private/NetworkFileServerHttp.cpp
+++ b/Engine/Source/Runtime/NetworkFileSystem/Private/NetworkFileServerHttp.cpp
@@ -18,8 +18,8 @@ class FNetworkFileServerClientConnectionHTTP : public FNetworkFileServerClientCo
 {
 
 public:
-	FNetworkFileServerClientConnectionHTTP(const FNetworkFileDelegateContainer* NetworkFileDelegates, const TArray<ITargetPlatform*>& InActiveTargetPlatforms )
-		:  FNetworkFileServerClientConnection( NetworkFileDelegates, InActiveTargetPlatforms)
+	FNetworkFileServerClientConnectionHTTP(const FNetworkFileServerOptions& FileServerOptions)
+		:  FNetworkFileServerClientConnection(FileServerOptions)
 	{
 	}
 
@@ -79,19 +79,14 @@ static struct lws_protocols Protocols[] = {
 //////////////////////////////////////////////////////////////////////////
 
 
-FNetworkFileServerHttp::FNetworkFileServerHttp(
-	int32 InPort,
-	FNetworkFileDelegateContainer InNetworkFileDelegateContainer,
-	const TArray<ITargetPlatform*>& InActiveTargetPlatforms
-	)
-	: ActiveTargetPlatforms(InActiveTargetPlatforms)
+FNetworkFileServerHttp::FNetworkFileServerHttp(FNetworkFileServerOptions InFileServerOptions)
+	: FileServerOptions(MoveTemp(InFileServerOptions))
 	, SslContext(nullptr)
 	, Context(nullptr)
-	, Port(InPort)
 {
-	if (Port < 0 )
+	if (FileServerOptions.Port < 0 )
 	{
-		Port = DEFAULT_HTTP_FILE_SERVING_PORT;
+		FileServerOptions.Port = DEFAULT_HTTP_FILE_SERVING_PORT;
 	}
 
 #if WITH_SSL
@@ -106,8 +101,6 @@ FNetworkFileServerHttp::FNetworkFileServerHttp(
 
 	UE_LOG(LogFileServer, Warning, TEXT("Unreal Network Http File Server starting up..."));
 
-	NetworkFileDelegates = InNetworkFileDelegateContainer;
-
 	StopRequested.Reset();
 	Ready.Reset();
 
@@ -140,7 +133,7 @@ bool FNetworkFileServerHttp::Init()
 	memset(&Info,0,sizeof(lws_context_creation_info));
 
 	// look up libwebsockets.h for details.
-	Info.port = Port;
+	Info.port = FileServerOptions.Port;
 	// we listen on all available interfaces.
 	Info.iface = NULL;
 	// serve only the http protocols.
@@ -168,10 +161,10 @@ bool FNetworkFileServerHttp::Init()
 
 	Context = lws_create_context(&Info);
 
-	Port = Info.port;
+	FileServerOptions.Port = Info.port;
 
 	if (Context == NULL) {
-		UE_LOG(LogFileServer, Error, TEXT(" Could not create a libwebsocket context.\n Port : %d is already in use.\n Exiting...\n"), Port);
+		UE_LOG(LogFileServer, Error, TEXT(" Could not create a libwebsocket context.\n Port : %d is already in use.\n Exiting...\n"), FileServerOptions.Port);
 		return false;
 	}
 
@@ -197,7 +190,7 @@ bool FNetworkFileServerHttp::GetAddressList(TArray<TSharedPtr<FInternetAddr> >&
 	// Fix up ports.
 	for (int32 AddressIndex = 0; AddressIndex < OutAddresses.Num(); ++AddressIndex)
 	{
-		OutAddresses[AddressIndex]->SetPort(Port);
+		OutAddresses[AddressIndex]->SetPort(FileServerOptions.Port);
 	}
 
 	return true;
@@ -232,7 +225,7 @@ void FNetworkFileServerHttp::Shutdown()
 
 uint32 FNetworkFileServerHttp::Run()
 {
-	UE_LOG(LogFileServer, Display, TEXT("Unreal Network File Http Server is ready for client connections on port %d"), Port);
+	UE_LOG(LogFileServer, Display, TEXT("Unreal Network File Http Server is ready for client connections on port %d"), FileServerOptions.Port);
 
 	// start servicing.
 
@@ -276,7 +269,7 @@ FNetworkFileServerHttp::~FNetworkFileServerHttp()
 
 FNetworkFileServerClientConnectionHTTP* FNetworkFileServerHttp::CreateNewConnection()
 {
-	return new FNetworkFileServerClientConnectionHTTP(&NetworkFileDelegates,ActiveTargetPlatforms);
+	return new FNetworkFileServerClientConnectionHTTP(FileServerOptions);
 }
 
 // Have a similar process function for the normal tcp connection.
diff --git a/Engine/Source/Runtime/NetworkFileSystem/Private/NetworkFileServerHttp.h b/Engine/Source/Runtime/NetworkFileSystem/Private/NetworkFileServerHttp.h
--- a/Engine/Source/Runtime/NetworkFileSystem/Private/NetworkFileServerHttp.h
+++ b/Engine/Source/Runtime/NetworkFileSystem/Private/NetworkFileServerHttp.h
@@ -42,7 +42,7 @@ class FNetworkFileServerHttp
 {
 
 public:
-	FNetworkFileServerHttp(int32 InPort, FNetworkFileDelegateContainer InNetworkFileDelegateContainer, const TArray<ITargetPlatform*>& InActiveTargetPlatforms );
+	FNetworkFileServerHttp(FNetworkFileServerOptions InFileServerOptions);
 
 	// INetworkFileServer Interface.
 
@@ -74,10 +74,8 @@ private:
 	// factory method for creating a new Client Connection.
 	class FNetworkFileServerClientConnectionHTTP* CreateNewConnection();
 
-	FNetworkFileDelegateContainer NetworkFileDelegates;
-
-	// cached copy of the active target platforms (if any)
-	const TArray<ITargetPlatform*> ActiveTargetPlatforms;
+	// File server options
+	FNetworkFileServerOptions FileServerOptions;
 
 	/** OpenSSL context */
 	SSL_CTX* SslContext;
@@ -88,9 +86,6 @@ private:
 	// Service Http connections on this thread.
 	FRunnableThread* WorkerThread;
 
-	// port on which this http server runs.
-	int32 Port;
-
 	// used to send simple message.
 	FThreadSafeCounter StopRequested;
 
diff --git a/Engine/Source/Runtime/NetworkFileSystem/Private/NetworkFileSystemModule.cpp b/Engine/Source/Runtime/NetworkFileSystem/Private/NetworkFileSystemModule.cpp
--- a/Engine/Source/Runtime/NetworkFileSystem/Private/NetworkFileSystemModule.cpp
+++ b/Engine/Source/Runtime/NetworkFileSystem/Private/NetworkFileSystemModule.cpp
@@ -23,9 +23,19 @@ public:
 
 	// INetworkFileSystemModule interface
 
-	virtual INetworkFileServer* CreateNetworkFileServer( bool bLoadTargetPlatforms, int32 Port, FNetworkFileDelegateContainer NetworkFileDelegateContainer, const ENetworkFileServerProtocol Protocol ) const override
+	virtual INetworkFileServer* CreateNetworkFileServer( bool bLoadTargetPlatforms, int32 Port, FNetworkFileDelegateContainer NetworkFileDelegateContainer, const ENetworkFileServerProtocol Protocol ) const
+	{
+		FNetworkFileServerOptions FileServerOptions;
+		FileServerOptions.Protocol = Protocol;
+		FileServerOptions.Port = Port;
+		FileServerOptions.Delegates = MoveTemp(NetworkFileDelegateContainer);
+		FileServerOptions.bRestrictPackageAssetsToSandbox = false; 
+		
+		return CreateNetworkFileServer(MoveTemp(FileServerOptions), bLoadTargetPlatforms);
+	}
+	
+	virtual INetworkFileServer* CreateNetworkFileServer(FNetworkFileServerOptions FileServerOptions, bool bLoadTargetPlatforms) const override
 	{
-		TArray<ITargetPlatform*> ActiveTargetPlatforms;
 		if (bLoadTargetPlatforms)
 		{
 			ITargetPlatformManagerModule& TPM = GetTargetPlatformManagerRef();
@@ -34,25 +44,25 @@ public:
 			FString Platforms;
 			if (FParse::Value(FCommandLine::Get(), TEXT("TARGETPLATFORM="), Platforms))
 			{
-				ActiveTargetPlatforms =  TPM.GetActiveTargetPlatforms();
+				FileServerOptions.TargetPlatforms =  TPM.GetActiveTargetPlatforms();
 			}
 			else
 			{
-				ActiveTargetPlatforms = TPM.GetTargetPlatforms();
+				FileServerOptions.TargetPlatforms = TPM.GetTargetPlatforms();
 			}
 		}
 
-		switch ( Protocol )
+		switch (FileServerOptions.Protocol)
 		{
 #if ENABLE_HTTP_FOR_NFS
 		case NFSP_Http: 
-			return new FNetworkFileServerHttp(Port, NetworkFileDelegateContainer, ActiveTargetPlatforms);
+			return new FNetworkFileServerHttp(MoveTemp(FileServerOptions));
 #endif
 		case NFSP_Tcp:
-			return new FNetworkFileServer(Port, NetworkFileDelegateContainer, ActiveTargetPlatforms);
+			return new FNetworkFileServer(MoveTemp(FileServerOptions));
 		}
  
-		return NULL;
+		return nullptr;
 	}
 };
 
diff --git a/Engine/Source/Runtime/NetworkFileSystem/Public/INetworkFileSystemModule.h b/Engine/Source/Runtime/NetworkFileSystem/Public/INetworkFileSystemModule.h
--- a/Engine/Source/Runtime/NetworkFileSystem/Public/INetworkFileSystemModule.h
+++ b/Engine/Source/Runtime/NetworkFileSystem/Public/INetworkFileSystemModule.h
@@ -107,13 +107,31 @@ public:
 	FOnFileModifiedDelegate* OnFileModifiedCallback; // this is called from other systems to notify the network file system that a file has been modified hence the terminology callback
 };
 
-
 enum ENetworkFileServerProtocol
 {
 	NFSP_Tcp,
 	NFSP_Http,
 };
 
+// Network file server options
+struct FNetworkFileServerOptions
+{
+	/* File server protocol*/
+	ENetworkFileServerProtocol Protocol = NFSP_Tcp;
+
+	/* The port number to bind to (-1 = default port, 0 = any available port) */
+	int32 Port = INDEX_NONE;
+
+	/* Optional delegates to be invoked when a file is requested by a client */
+	FNetworkFileDelegateContainer Delegates;
+
+	/* Active target platform(s) */
+	TArray<ITargetPlatform*> TargetPlatforms;
+
+	/* When running cook on the fly this options restricts package assets being sent from the project content folder */
+	bool bRestrictPackageAssetsToSandbox = false;
+};
+
 /**
  * Interface for network file system modules.
  */
@@ -134,6 +152,16 @@ public:
 	 */
 	virtual INetworkFileServer* CreateNetworkFileServer( bool bLoadTargetPlatforms, int32 Port = -1, FNetworkFileDelegateContainer InNetworkFileDelegateContainer = FNetworkFileDelegateContainer(), const ENetworkFileServerProtocol Protocol = NFSP_Tcp ) const = 0;
 
+    /**
+	 * Creates a new network file server.
+	 *
+	 * @param FileServerOptions File server options
+	 * @param bLoadTargetPlatforms If true, gets the target platform from the command line or all available target platforms
+	 *
+	 * @return The new file server, or nullptr if creation failed.
+	 */
+	virtual INetworkFileServer* CreateNetworkFileServer(FNetworkFileServerOptions FileServerOptions, bool bLoadTargetPlatforms) const = 0;
+
 public:
 
 	/**
