COTF: Fixed bug when syncing files with restricted sandbox

#rnx

diff --git a/Engine/Source/Runtime/NetworkFileSystem/Private/NetworkFileServerConnection.cpp b/Engine/Source/Runtime/NetworkFileSystem/Private/NetworkFileServerConnection.cpp
--- a/Engine/Source/Runtime/NetworkFileSystem/Private/NetworkFileServerConnection.cpp
+++ b/Engine/Source/Runtime/NetworkFileSystem/Private/NetworkFileServerConnection.cpp
@@ -45,6 +45,21 @@ static FString MakeAbsoluteNormalizedDir(const FString& InPath)
 	return Out;
 }
 
+struct FSandboxOnlyScope
+{
+	FSandboxOnlyScope(FSandboxPlatformFile& InSandbox, bool bInSandboxOnly)
+		: Sandbox(InSandbox)
+	{
+		Sandbox.SetSandboxOnly(bInSandboxOnly);
+	}
+
+	~FSandboxOnlyScope()
+	{
+		Sandbox.SetSandboxOnly(false);
+	}
+
+	FSandboxPlatformFile& Sandbox;
+};
 
 /* FNetworkFileServerClientConnection structors
  *****************************************************************************/
@@ -459,11 +474,8 @@ void FNetworkFileServerClientConnection::ProcessOpenFile( FArchive& In, FArchive
 	NetworkFileDelegates->FileRequestDelegate.ExecuteIfBound(Filename, ConnectedPlatformName, NewUnsolictedFiles);
 
 	// Disable access to outside the sandbox to prevent sending uncooked packages to the client
-	if (bRestrictPackageAssetsToSandbox && !bIsWriting)
-	{
-		const bool bIsPackageAsset = FPackageName::IsPackageExtension(*FPaths::GetExtension(Filename, true));
-		Sandbox->SetSandboxOnly(bIsPackageAsset);
-	}
+	const bool bSandboxOnly = bRestrictPackageAssetsToSandbox && !bIsWriting && FPackageName::IsPackageExtension(*FPaths::GetExtension(Filename, true));
+	FSandboxOnlyScope _(*Sandbox, bSandboxOnly);
 
 	FDateTime ServerTimeStamp = Sandbox->GetTimeStamp(*Filename);
 	int64 ServerFileSize = 0;
@@ -478,8 +490,6 @@ void FNetworkFileServerClientConnection::ProcessOpenFile( FArchive& In, FArchive
 		ServerFileSize = File->Size();
 	}
 
-	Sandbox->SetSandboxOnly(false);
-
 	uint64 HandleId = ++LastHandleId;
 	OpenFiles.Add( HandleId, File );
 	
@@ -1218,11 +1228,8 @@ bool FNetworkFileServerClientConnection::PackageFile( FString& Filename, FString
 	FDateTime ServerTimeStamp = Sandbox->GetTimeStamp(*Filename);
 
 	// Disable access to outside the sandbox to prevent sending uncooked packages to the client
-	if (bRestrictPackageAssetsToSandbox)
-	{
-		const bool bIsPackageAsset = FPackageName::IsPackageExtension(*FPaths::GetExtension(Filename, true));
-		Sandbox->SetSandboxOnly(bIsPackageAsset);
-	}
+	const bool bSandboxOnly = bRestrictPackageAssetsToSandbox && FPackageName::IsPackageExtension(*FPaths::GetExtension(Filename, true));
+	FSandboxOnlyScope _(*Sandbox, bSandboxOnly);
 
 	FString AbsHostFile = Sandbox->ConvertToAbsolutePathForExternalAppForRead(*Filename);
 	if (ConnectedTargetPlatform != nullptr && ConnectedTargetPlatform->CopyFileToTarget(ConnectedIPAddress, AbsHostFile, TargetFilename, ConnectedTargetCustomData))
@@ -1271,8 +1278,6 @@ bool FNetworkFileServerClientConnection::PackageFile( FString& Filename, FString
 		UE_LOG(LogFileServer, Display, TEXT("Read %s, %d bytes"), *Filename, Contents.Num());
 	}
 
-	Sandbox->SetSandboxOnly(false);
-
 	Out << Filename;
 	Out << ServerTimeStamp;
 	uint64 FileSize = Contents.Num();
