CookOnTheFly (platform protocol): Editor: Added support of the platform-specific transport communication protocols to the editor and use it in the cooking server.

This change allows the cooking server to communicate with clients using platform-specific communication protocols instead of TCP (the default).

This new mode is activated by adding '-platformprotocol' command-line parameter when starting the cooking server (in addition to the usual '-cookonthefly').

The new functionality is included in FNetworkFileServerPlatformProtocol, which is a variant of the preexisting FNetworkFileServer but using ITargetDevice/ITargetDeviceSocket interfaces instead of traditional sockets.

Direct communication with targets is only supported on some platforms, which can be checked by calling ITargetPlatform::SupportsFeature(ETargetPlatformFeatures::DirectDataExchange).

diff --git a/Engine/Source/Developer/TargetPlatform/Public/Interfaces/ITargetDevice.h b/Engine/Source/Developer/TargetPlatform/Public/Interfaces/ITargetDevice.h
--- a/Engine/Source/Developer/TargetPlatform/Public/Interfaces/ITargetDevice.h
+++ b/Engine/Source/Developer/TargetPlatform/Public/Interfaces/ITargetDevice.h
@@ -4,6 +4,7 @@
 
 #include "CoreMinimal.h"
 #include "Interfaces/TargetDeviceId.h"
+#include "Interfaces/ITargetDeviceSocket.h"
 
 class ITargetDevice;
 class ITargetDeviceOutput;
@@ -398,6 +399,42 @@ public:
 		return ExecuteConsoleCommand(TEXT("ReloadGlobalShaders"));
 	}
 
+	/**
+	 * Opens a direct connection with the device allowing data exchange with a process running on the target.
+	 * 
+	 * ProtocolIndex is a number that identifies the connection and has to be know both to the target
+	 * process and the PC. There can be only one connection using a given protocol index at a time.
+	 * 
+	 * You may check EHostProtocol enumeration for known protocols used by the engine.
+	 * 
+	 * This function just opens a communication channel but doesn't check if there is a peer
+	 * on the other end of the connection. To know this, you should either check IsProtocolAvailable
+	 * or ITargetDeviceSocket::Connected.
+	 * 
+	 * @param ProtocolIndex Unique index of the communication channel (from 0 to a platform-dependent maximum).
+	 * @return Socket object on success, nullptr otherwise.
+	 * @see CloseSocket, EHostProtocol, IsProtocolAvailable, ITargetDeviceSocket::Connected
+	 */
+	virtual ITargetDeviceSocketPtr OpenConnection(uint32 ProtocolIndex)	{ return nullptr; }
+
+	/**
+	 * Closes a previously opened connection with a process on the target.
+	 * 
+	 * @param Socket Socket object returned by OpenSocket.
+	 * @see OpenSocket
+	 */
+	virtual void CloseConnection(ITargetDeviceSocketPtr Socket)			{ }
+
+	/**
+	 * Checks if connections using the given ProtocolIndex are available for this device at the moment.
+	 * If this function returns true, it means that OpenConnection should succeed and communication
+	 * with a game running on the target should be possible.
+	 * 
+	 * @param ProtocolIndex Unique index of the communication channel (from 0 to a platform-dependent maximum).
+	 * @return True if the protocol is available and we can connect to the device, false otherwise.
+	 */
+	virtual bool IsProtocolAvailable(uint32 ProtocolIndex) const		{ return false; }
+
 public:
 
 	/**
diff --git a/Engine/Source/Developer/TargetPlatform/Public/Interfaces/ITargetDeviceSocket.h b/Engine/Source/Developer/TargetPlatform/Public/Interfaces/ITargetDeviceSocket.h
new file mode 100644
--- /dev/null
+++ b/Engine/Source/Developer/TargetPlatform/Public/Interfaces/ITargetDeviceSocket.h
@@ -0,0 +1,65 @@
+// Copyright Epic Games, Inc. All Rights Reserved.
+
+#pragma once
+
+
+/**
+ * Interface for target device sockets.
+ *
+ * This interface provides an abstraction for communicating with processes running on the target.
+ */
+class ITargetDeviceSocket
+{
+public:
+
+	/**
+	 * Send data to a connected process on the target device.
+	 * 
+	 * This is a blocking operation and it will return only after the whole buffer
+	 * has been sent or an error occurs (e.g. the target disconnects).
+	 * 
+	 * False return value typically indicates that the remote peer has closed the connection 
+	 * on their side. If it happens, this socket should be closed as well.
+	 *
+	 * @param Data        Buffer containing data to be sent (.
+	 * @param BytesToSend The number of bytes from Data that are to be sent.
+	 * @return true if communication was successful, false otherwise.
+	 * 
+	 * @see IHostDevice::OpenConnection, IHostDevice::CloseConnection
+	 */
+	virtual bool Send(const void* Data, uint64 BytesToSend) = 0;
+
+	/**
+	 * Receive data from a connected process on the target device.
+	 * 
+	 * This is a blocking operation and it will return only after the expected amount
+	 * of data has been received or an error occurs (e.g. the target disconnects).
+	 * 
+	 * False return value typically indicates that the remote peer has closed the connection
+	 * on their side. If it happens, this socket should be closed as well.
+	 *
+	 * @param Data           Target buffer for the data to receive (it has to be large enough to store BytesToReceive).
+	 * @param BytesToReceive The number of bytes to receive and store in Data.
+	 * @return true if communication was successful, false otherwise.
+	 * 
+	 * @see IHostDevice::OpenConnection, IHostDevice::CloseConnection
+	 */
+	virtual bool Receive(void* Data, uint64 BytesToReceive) = 0;
+
+	/**
+	 * Returns true if this socket is actually connected to another peer and is ready to send/receive data.
+	 */
+	virtual bool Connected() const = 0;
+
+public:
+
+	/** Virtual destructor. */
+	virtual ~ITargetDeviceSocket() { }
+
+};
+
+// Type definition for shared references to instances of IPlatformHostSocket.
+typedef TSharedRef<ITargetDeviceSocket, ESPMode::ThreadSafe> ITargetDeviceSocketRef;
+
+// Type definition for shared pointers to instances of IPlatformHostSocket.
+typedef TSharedPtr<ITargetDeviceSocket, ESPMode::ThreadSafe> ITargetDeviceSocketPtr;
diff --git a/Engine/Source/Developer/TargetPlatform/Public/Interfaces/ITargetPlatform.h b/Engine/Source/Developer/TargetPlatform/Public/Interfaces/ITargetPlatform.h
--- a/Engine/Source/Developer/TargetPlatform/Public/Interfaces/ITargetPlatform.h
+++ b/Engine/Source/Developer/TargetPlatform/Public/Interfaces/ITargetPlatform.h
@@ -107,6 +107,9 @@ enum class ETargetPlatformFeatures
 
 	/* The platform makes use of extra cook-time file region metadata in its packaging process. */
 	CookFileRegionMetadata,
+
+	/** The platform supports communication (reading and writing data) between a target a connected PC. */
+	DirectDataExchange,
 };
 
 enum class EPlatformAuthentication
diff --git a/Engine/Source/Editor/UnrealEd/Classes/CookOnTheSide/CookOnTheFlyServer.h b/Engine/Source/Editor/UnrealEd/Classes/CookOnTheSide/CookOnTheFlyServer.h
--- a/Engine/Source/Editor/UnrealEd/Classes/CookOnTheSide/CookOnTheFlyServer.h
+++ b/Engine/Source/Editor/UnrealEd/Classes/CookOnTheSide/CookOnTheFlyServer.h
@@ -153,6 +153,8 @@ public:
 		bool bBindAnyPort = false;
 		/** Whether to save the cooked output to the Zen storage server. */
 		bool bZenStore = false;
+		/** Whether the network file server should use a platform-specific communication protocol instead of TCP (used when bZenStore == false) */
+		bool bPlatformProtocol = false;
 		/** Target platforms */
 		TArray<ITargetPlatform*> TargetPlatforms;
 	};
diff --git a/Engine/Source/Editor/UnrealEd/Private/Commandlets/CookCommandlet.cpp b/Engine/Source/Editor/UnrealEd/Private/Commandlets/CookCommandlet.cpp
--- a/Engine/Source/Editor/UnrealEd/Private/Commandlets/CookCommandlet.cpp
+++ b/Engine/Source/Editor/UnrealEd/Private/Commandlets/CookCommandlet.cpp
@@ -426,6 +426,7 @@ bool UCookCommandlet::CookOnTheFly( FGuid InstanceId, int32 Timeout, bool bForce
 	UCookOnTheFlyServer::FCookOnTheFlyOptions CookOnTheFlyStartupOptions;
 	CookOnTheFlyStartupOptions.bBindAnyPort = InstanceId.IsValid();
 	CookOnTheFlyStartupOptions.bZenStore = Switches.Contains(TEXT("ZenStore"));
+	CookOnTheFlyStartupOptions.bPlatformProtocol = Switches.Contains(TEXT("PlatformProtocol"));
 	CookOnTheFlyStartupOptions.TargetPlatforms = TargetPlatforms;
 
 	if (CookOnTheFlyServer->StartCookOnTheFly(MoveTemp(CookOnTheFlyStartupOptions)) == false)
diff --git a/Engine/Source/Editor/UnrealEd/Private/CookOnTheFlyServer.cpp b/Engine/Source/Editor/UnrealEd/Private/CookOnTheFlyServer.cpp
--- a/Engine/Source/Editor/UnrealEd/Private/CookOnTheFlyServer.cpp
+++ b/Engine/Source/Editor/UnrealEd/Private/CookOnTheFlyServer.cpp
@@ -525,7 +525,7 @@ bool UCookOnTheFlyServer::StartCookOnTheFly(FCookOnTheFlyOptions InCookOnTheFlyO
 	else
 	{
 		FNetworkFileServerOptions ServerOptions;
-		ServerOptions.Protocol = NFSP_Tcp;
+		ServerOptions.Protocol = CookOnTheFlyOptions.bPlatformProtocol ? NFSP_Platform : NFSP_Tcp;
 		ServerOptions.Port = CookOnTheFlyOptions.bBindAnyPort ? 0 : -1;
 		CookOnTheFlyRequestManager = UE::Cook::MakeNetworkFileCookOnTheFlyRequestManager(*CookOnTheFlyServerInterface, MoveTemp(ServerOptions));
 	}
diff --git a/Engine/Source/Runtime/NetworkFileSystem/Private/NetworkFileServerPlatformProtocol.cpp b/Engine/Source/Runtime/NetworkFileSystem/Private/NetworkFileServerPlatformProtocol.cpp
new file mode 100644
--- /dev/null
+++ b/Engine/Source/Runtime/NetworkFileSystem/Private/NetworkFileServerPlatformProtocol.cpp
@@ -0,0 +1,343 @@
+// Copyright Epic Games, Inc. All Rights Reserved.
+
+#include "NetworkFileServerPlatformProtocol.h"
+#include "HAL/RunnableThread.h"
+#include "Misc/OutputDeviceRedirector.h"
+#include "NetworkMessage.h"
+#include "NetworkFileSystemLog.h"
+#include "NetworkFileServerConnection.h"
+#include "Interfaces/ITargetPlatform.h"
+#include "Interfaces/ITargetDevice.h"
+#include "Interfaces/ITargetDeviceSocket.h"
+#include "GenericPlatform/GenericPlatformHostCommunication.h"
+
+
+
+namespace
+{
+	class FSimpleAbstractSocket_PlatformProtocol : public FSimpleAbstractSocket
+	{
+	public:
+
+		FSimpleAbstractSocket_PlatformProtocol(ITargetDeviceSocketPtr InSocket)
+			: Socket(InSocket)
+		{
+			check(Socket != nullptr);
+		}
+
+		virtual bool Receive(uint8* Results, int32 Size) const override
+		{
+			return Socket->Receive(Results, Size);
+		}
+
+		virtual bool Send(const uint8* Buffer, int32 Size) const override
+		{
+			return Socket->Send(Buffer, Size);
+		}
+
+		virtual uint32 GetMagic() const override
+		{
+			return 0x9E2B83C7;
+		}
+
+	private:
+
+		ITargetDeviceSocketPtr Socket;
+	};
+}
+
+
+class FNetworkFileServerPlatformProtocol::FConnectionThreaded
+	: public FNetworkFileServerClientConnection
+	, protected FRunnable 
+{
+public:
+
+	FConnectionThreaded(ITargetDevicePtr InDevice, ITargetDeviceSocketPtr InSocket, const FNetworkFileServerOptions& InOptions)
+		: FNetworkFileServerClientConnection(InOptions)
+		, Device(InDevice)
+		, Socket(InSocket)
+	{
+		Running       = true;
+		StopRequested = false;
+	
+#if UE_BUILD_DEBUG
+		// This thread needs more space in debug builds as it tries to log messages and such.
+		const static uint32 NetworkFileServerThreadSize = 2 * 1024 * 1024; 
+#else
+		const static uint32 NetworkFileServerThreadSize = 1 * 1024 * 1024; 
+#endif
+		WorkerThread = FRunnableThread::Create(this, TEXT("FNetworkFileServerCustomClientConnection"), NetworkFileServerThreadSize, TPri_AboveNormal);
+	}
+
+
+	virtual bool Init() override
+	{
+#if PLATFORM_WINDOWS
+		FWindowsPlatformMisc::CoInitialize(ECOMModel::Multithreaded);
+#endif
+
+		return true; 
+	}
+
+	virtual uint32 Run() override
+	{
+		while (!StopRequested)
+		{
+			// Read a header and payload pair.
+			FArrayReader Payload; 
+			if (!FNFSMessageHeader::ReceivePayload(Payload, FSimpleAbstractSocket_PlatformProtocol(Socket)))
+				break; 
+
+			// Now process the contents of the payload.
+			if ( !FNetworkFileServerClientConnection::ProcessPayload(Payload) )
+			{
+				// Give the processing of the payload a chance to terminate the connection
+				// failed to process message.
+				UE_LOG(LogFileServer, Warning, TEXT("Unable to process payload terminating connection"));
+				break;
+			}
+		}
+
+		return true;
+	}
+
+	virtual void Stop() override
+	{
+		StopRequested = true;
+	}
+
+	virtual void Exit() override
+	{
+		Device->CloseConnection(Socket);
+		Socket = nullptr;
+
+#if PLATFORM_WINDOWS
+		FWindowsPlatformMisc::CoUninitialize();
+#endif
+
+		Running = false; 
+	}
+
+	virtual bool SendPayload( TArray<uint8> &Out ) override
+	{
+		return FNFSMessageHeader::WrapAndSendPayload(Out, FSimpleAbstractSocket_PlatformProtocol(Socket));
+	}
+
+	bool IsRunning()
+	{
+		return Running; 
+	}
+
+	FString GetName() const
+	{
+		return FString::Printf(TEXT("%s (%s)"), *Device->GetName(), *Device->GetTargetPlatform().PlatformName());
+	}
+
+	ITargetDevicePtr GetDevice() const
+	{
+		return Device;
+	}
+
+	~FConnectionThreaded()
+	{
+		WorkerThread->Kill(true);
+	}
+
+private:
+
+	ITargetDevicePtr	   Device;
+	ITargetDeviceSocketPtr Socket;
+
+	std::atomic<bool>      StopRequested;
+	std::atomic<bool>      Running;
+	FRunnableThread*       WorkerThread; 
+};
+
+
+FNetworkFileServerPlatformProtocol::FNetworkFileServerPlatformProtocol(FNetworkFileServerOptions InFileServerOptions)
+	: FileServerOptions(MoveTemp(InFileServerOptions))
+{
+	Running       = false;
+	StopRequested = false;
+
+	UE_LOG(LogFileServer, Display , TEXT("Unreal Network File Server (custom protocol) starting up..."));
+
+	// Check the list of platforms once on start (any missing platforms will be ignored later on to avoid spamming the log).
+	for (ITargetPlatform* TargetPlatform : FileServerOptions.TargetPlatforms)
+	{
+		if (!TargetPlatform->SupportsFeature(ETargetPlatformFeatures::DirectDataExchange))
+		{
+			UE_LOG(LogFileServer, Error, TEXT("Platform '%s' does not support direct communication with targets (it will be ignored)."), *TargetPlatform->PlatformName());
+		}
+	}
+
+	// Create a thread that will be updating the list of connected target devices.
+	Thread = FRunnableThread::Create(this, TEXT("FNetworkFileServerCustomProtocol"), 8 * 1024, TPri_AboveNormal);
+
+	UE_LOG(LogFileServer, Display, TEXT("Unreal Network File Server is ready for client connections!"));
+}
+
+
+FNetworkFileServerPlatformProtocol::~FNetworkFileServerPlatformProtocol()
+{
+	// Kill the running thread.
+	if (Thread != nullptr)
+	{
+		Thread->Kill(true);
+
+		delete Thread;
+		Thread = nullptr;
+	}
+}
+
+
+uint32 FNetworkFileServerPlatformProtocol::Run()
+{
+#if PLATFORM_WINDOWS
+	FWindowsPlatformMisc::CoInitialize(ECOMModel::Multithreaded);
+#endif
+
+	Running = true; 
+
+	// Go until requested to be done.
+	while (!StopRequested)
+	{
+		UpdateConnections();
+
+		FPlatformProcess::Sleep(1.0f);
+	}
+
+#if PLATFORM_WINDOWS
+	FWindowsPlatformMisc::CoUninitialize();
+#endif
+
+	return 0;
+}
+
+
+void FNetworkFileServerPlatformProtocol::Stop()
+{
+	StopRequested = true;
+}
+
+
+void FNetworkFileServerPlatformProtocol::Exit()
+{
+	// Close all connections.
+	for (auto* Connection : Connections)
+	{
+		delete Connection;
+	}
+
+	Connections.Empty();
+}
+
+
+FString FNetworkFileServerPlatformProtocol::GetSupportedProtocol() const
+{
+	return FString("custom");
+}
+
+
+bool FNetworkFileServerPlatformProtocol::GetAddressList( TArray<TSharedPtr<FInternetAddr> >& OutAddresses ) const
+{
+	return 0;
+}
+
+
+bool FNetworkFileServerPlatformProtocol::IsItReadyToAcceptConnections() const
+{
+	return Running; 
+}
+
+
+int32 FNetworkFileServerPlatformProtocol::NumConnections() const
+{
+	return Connections.Num(); 
+}
+
+
+void FNetworkFileServerPlatformProtocol::Shutdown()
+{
+	Stop();
+}
+
+
+void FNetworkFileServerPlatformProtocol::UpdateConnections()
+{
+	RemoveClosedConnections();
+
+	AddConnectionsForNewDevices();
+}
+
+
+void FNetworkFileServerPlatformProtocol::RemoveClosedConnections()
+{
+	for (int32 ConnectionIndex = 0; ConnectionIndex < Connections.Num(); ++ConnectionIndex)
+	{
+		FConnectionThreaded* Connection = Connections[ConnectionIndex];
+
+		if (!Connection->IsRunning())
+		{
+			UE_LOG(LogFileServer, Display, TEXT("Client %s disconnected."), *Connection->GetName());
+			Connections.RemoveAtSwap(ConnectionIndex);
+			delete Connection;
+		}
+	}
+}
+
+
+void FNetworkFileServerPlatformProtocol::AddConnectionsForNewDevices()
+{
+	for (ITargetPlatform* TargetPlatform : FileServerOptions.TargetPlatforms)
+	{
+		if (TargetPlatform->SupportsFeature(ETargetPlatformFeatures::DirectDataExchange))
+		{
+			AddConnectionsForNewDevices(TargetPlatform);
+		}
+	}
+}
+
+
+void FNetworkFileServerPlatformProtocol::AddConnectionsForNewDevices(ITargetPlatform* TargetPlatform)
+{
+	TArray<ITargetDevicePtr> TargetDevices;
+
+	TargetPlatform->GetAllDevices(TargetDevices);
+
+	for (ITargetDevicePtr Device : TargetDevices)
+	{
+		if (Device->IsConnected())
+		{
+			FConnectionThreaded** ExistingConnection =
+				Connections.FindByPredicate(
+					[&Device](const FConnectionThreaded* Connection)
+					{
+						return Connection->GetDevice() == Device;
+					}
+			);
+
+			// Checking IsProtocolAvailable first would make more sense, but internally it queries
+			// COM interfaces, which throws exceptions if the protocol is already in use. While we catch
+			// and process these exceptions, Visual Studio intercepts then as well and outputs messages
+			// spamming the log, which hinders the debugging experience.
+			if (!ExistingConnection)
+			{
+				if (Device->IsProtocolAvailable(EHostProtocol::CookOnTheFly))
+				{
+					ITargetDeviceSocketPtr Socket = Device->OpenConnection(EHostProtocol::CookOnTheFly);
+
+					if (Socket)
+					{
+						FConnectionThreaded* Connection =
+							new FConnectionThreaded(Device, Socket, FileServerOptions);
+
+						Connections.Add(Connection);
+						UE_LOG(LogFileServer, Display, TEXT("Client %s connected."), *Connection->GetName());
+					}
+				}
+			}
+		}
+	}
+}
diff --git a/Engine/Source/Runtime/NetworkFileSystem/Private/NetworkFileServerPlatformProtocol.h b/Engine/Source/Runtime/NetworkFileSystem/Private/NetworkFileServerPlatformProtocol.h
new file mode 100644
--- /dev/null
+++ b/Engine/Source/Runtime/NetworkFileSystem/Private/NetworkFileServerPlatformProtocol.h
@@ -0,0 +1,87 @@
+// Copyright Epic Games, Inc. All Rights Reserved.
+
+#pragma once
+
+#include "CoreMinimal.h"
+#include "HAL/Runnable.h"
+#include "INetworkFileServer.h"
+#include "INetworkFileSystemModule.h"
+
+class ITargetPlatform;
+
+
+/**
+ * This class wraps the server thread and network connection.
+ * 
+ * It uses ITargetDeviceSocket to exchanging data with the connected targets.
+ * This interface is an abstraction for direct pc-target communication as provided
+ * by the platforms. IPlatformHostCommunication/IPlatformHostSocket are the
+ * corresponding interfaces used on the game side.
+ * 
+ * This implementation is based on FNetworkFileServer (the TCP-based server).
+ */
+class FNetworkFileServerPlatformProtocol
+	: public FRunnable
+	, public INetworkFileServer
+{
+public:
+
+	/**
+	 * Creates and initializes a new instance.
+	 *
+	 * @param InFileServerOptions Network file server options
+	 */
+	FNetworkFileServerPlatformProtocol(FNetworkFileServerOptions InFileServerOptions);
+
+	/**
+	 * Destructor.
+	 */
+	~FNetworkFileServerPlatformProtocol();
+
+public:
+
+	//~ Begin FRunnable Interface
+
+	virtual uint32 Run()  override;
+	virtual void   Stop() override;
+	virtual void   Exit() override;
+
+	//~ End FRunnable Interface
+
+public:
+
+	//~ Begin INetworkFileServer Interface
+
+	virtual bool    IsItReadyToAcceptConnections() const override;
+	virtual bool    GetAddressList(TArray<TSharedPtr<FInternetAddr> >& OutAddresses) const override;
+	virtual FString GetSupportedProtocol() const override;
+	virtual int32   NumConnections() const override;
+	virtual void    Shutdown() override;
+
+	//~ End INetworkFileServer Interface
+
+private:
+
+	void UpdateConnections();
+	void AddConnectionsForNewDevices();
+	void AddConnectionsForNewDevices(ITargetPlatform* TargetPlatform);
+	void RemoveClosedConnections();
+
+	class FConnectionThreaded;
+
+	// File server options
+	FNetworkFileServerOptions FileServerOptions;
+
+	// Holds the server thread object.
+	FRunnableThread* Thread;
+
+	// Holds the list of all client connections.
+	TArray<FConnectionThreaded*> Connections;
+
+	// Holds a flag indicating whether the thread should stop executing
+	std::atomic<bool> StopRequested;
+
+	// Is the Listener thread up and running. 
+	std::atomic<bool> Running;
+
+};
diff --git a/Engine/Source/Runtime/NetworkFileSystem/Private/NetworkFileSystemModule.cpp b/Engine/Source/Runtime/NetworkFileSystem/Private/NetworkFileSystemModule.cpp
--- a/Engine/Source/Runtime/NetworkFileSystem/Private/NetworkFileSystemModule.cpp
+++ b/Engine/Source/Runtime/NetworkFileSystem/Private/NetworkFileSystemModule.cpp
@@ -7,6 +7,7 @@
 #include "NetworkFileSystemLog.h"
 #include "NetworkFileServer.h"
 #include "NetworkFileServerHttp.h"
+#include "NetworkFileServerPlatformProtocol.h"
 #include "Interfaces/ITargetPlatformManagerModule.h"
 
 
@@ -60,6 +61,9 @@ public:
 #endif
 		case NFSP_Tcp:
 			return new FNetworkFileServer(MoveTemp(FileServerOptions));
+
+		case NFSP_Platform:
+			return new FNetworkFileServerPlatformProtocol(MoveTemp(FileServerOptions));
 		}
  
 		return nullptr;
diff --git a/Engine/Source/Runtime/NetworkFileSystem/Public/INetworkFileSystemModule.h b/Engine/Source/Runtime/NetworkFileSystem/Public/INetworkFileSystemModule.h
--- a/Engine/Source/Runtime/NetworkFileSystem/Public/INetworkFileSystemModule.h
+++ b/Engine/Source/Runtime/NetworkFileSystem/Public/INetworkFileSystemModule.h
@@ -65,6 +65,14 @@ enum ENetworkFileServerProtocol
 {
 	NFSP_Tcp,
 	NFSP_Http,
+
+	/**
+	 * Platform-specific type of connection between a target device and a host pc.
+	 * 
+	 * It potentially offers performance benefits compared to generic networking protocols but is
+	 * supported by only some of the platforms.
+	 */
+	NFSP_Platform,
 };
 
 // Network file server options
