- Fixed a COTF bug where it would use the TargetPlatformName (say WIndowsClient) instead of IniPlatformName (Windows) when looking in PLatformExtension directories

diff --git a/Engine/Source/Runtime/NetworkFileSystem/NetworkFileSystem.Build.cs b/Engine/Source/Runtime/NetworkFileSystem/NetworkFileSystem.Build.cs
--- a/Engine/Source/Runtime/NetworkFileSystem/NetworkFileSystem.Build.cs
+++ b/Engine/Source/Runtime/NetworkFileSystem/NetworkFileSystem.Build.cs
@@ -19,6 +19,7 @@ namespace UnrealBuildTool.Rules
 					"Projects",
 					"SandboxFile",
 					"TargetPlatform",
+					"DesktopPlatform",
 				});
 
 			PublicIncludePaths.AddRange(
diff --git a/Engine/Source/Runtime/NetworkFileSystem/Private/NetworkFileServerConnection.cpp b/Engine/Source/Runtime/NetworkFileSystem/Private/NetworkFileServerConnection.cpp
--- a/Engine/Source/Runtime/NetworkFileSystem/Private/NetworkFileServerConnection.cpp
+++ b/Engine/Source/Runtime/NetworkFileSystem/Private/NetworkFileServerConnection.cpp
@@ -15,6 +15,7 @@
 #include "Interfaces/ITargetPlatform.h"
 #include "HAL/PlatformTime.h"
 #include "Interfaces/IPluginManager.h"
+#include "PlatformInfo.h"
 
 
 /**
@@ -1083,8 +1084,10 @@ bool FNetworkFileServerClientConnection::ProcessGetFileList( FArchive& In, FArch
 	ConvertClientFilenameToServerFilename(ServerProjectPlatformExtensionsRelativePath);
 	for (const FString& TargetPlatform : TargetPlatformNames)
 	{
-		ScanExtensionRootDirectory(Sandbox.Get(), ServerEnginePlatformExtensionsRelativePath / TargetPlatform, RootDirectories, Visitor.FileTimes);
-		ScanExtensionRootDirectory(Sandbox.Get(), ServerProjectPlatformExtensionsRelativePath / TargetPlatform, RootDirectories, Visitor.FileTimes);
+		FString IniPlatformName = PlatformInfo::FindPlatformInfo(*TargetPlatform)->IniPlatformName.ToString();
+
+		ScanExtensionRootDirectory(Sandbox.Get(), ServerEnginePlatformExtensionsRelativePath / IniPlatformName, RootDirectories, Visitor.FileTimes);
+		ScanExtensionRootDirectory(Sandbox.Get(), ServerProjectPlatformExtensionsRelativePath / IniPlatformName, RootDirectories, Visitor.FileTimes);
 	}
 
 	UE_LOG(LogFileServer, Display, TEXT("Scanned server files, found %d files in %.2f seconds"), Visitor.FileTimes.Num(), FPlatformTime::Seconds() - FileScanStartTime);
