COTF: Make sure to read all platform extension directories related to the IniPlatformName.
Fix ShaderPlatform crash on platforms with a common folder that contains ShaderPlatform config data.

#rnx

diff --git a/Engine/Source/Runtime/NetworkFileSystem/Private/NetworkFileServerConnection.cpp b/Engine/Source/Runtime/NetworkFileSystem/Private/NetworkFileServerConnection.cpp
--- a/Engine/Source/Runtime/NetworkFileSystem/Private/NetworkFileServerConnection.cpp
+++ b/Engine/Source/Runtime/NetworkFileSystem/Private/NetworkFileServerConnection.cpp
@@ -1082,12 +1082,27 @@ bool FNetworkFileServerClientConnection::ProcessGetFileList( FArchive& In, FArch
 	ConvertClientFilenameToServerFilename(ServerEnginePlatformExtensionsRelativePath);
 	FString ServerProjectPlatformExtensionsRelativePath = ProjectPlatformExtensionsRelativePath;
 	ConvertClientFilenameToServerFilename(ServerProjectPlatformExtensionsRelativePath);
+
+	TArray<FString> PlatformDirectoryNames;
 	for (const FString& TargetPlatform : TargetPlatformNames)
 	{
-		FString IniPlatformName = PlatformInfo::FindPlatformInfo(*TargetPlatform)->IniPlatformName.ToString();
-
-		ScanExtensionRootDirectory(Sandbox.Get(), ServerEnginePlatformExtensionsRelativePath / IniPlatformName, RootDirectories, Visitor.FileTimes);
-		ScanExtensionRootDirectory(Sandbox.Get(), ServerProjectPlatformExtensionsRelativePath / IniPlatformName, RootDirectories, Visitor.FileTimes);
+		FName IniPlatformName = PlatformInfo::FindPlatformInfo(*TargetPlatform)->IniPlatformName;
+		const FDataDrivenPlatformInfo& PlatformInfo = FDataDrivenPlatformInfoRegistry::GetPlatformInfo(IniPlatformName);
+		PlatformDirectoryNames.Reserve(PlatformInfo.IniParentChain.Num() + PlatformInfo.AdditionalRestrictedFolders.Num() + 1);
+		PlatformDirectoryNames.Add(IniPlatformName.ToString());
+		for (const FString& PlatformName : PlatformInfo.AdditionalRestrictedFolders)
+		{
+			PlatformDirectoryNames.AddUnique(PlatformName);
+		}
+		for (const FString& PlatformName : PlatformInfo.IniParentChain)
+		{
+			PlatformDirectoryNames.AddUnique(PlatformName);
+		}
+	}
+	for (const FString& PlatformDirectoryName : PlatformDirectoryNames)
+	{
+		ScanExtensionRootDirectory(Sandbox.Get(), ServerEnginePlatformExtensionsRelativePath / PlatformDirectoryName, RootDirectories, Visitor.FileTimes);
+		ScanExtensionRootDirectory(Sandbox.Get(), ServerProjectPlatformExtensionsRelativePath / PlatformDirectoryName, RootDirectories, Visitor.FileTimes);
 	}
 
 	UE_LOG(LogFileServer, Display, TEXT("Scanned server files, found %d files in %.2f seconds"), Visitor.FileTimes.Num(), FPlatformTime::Seconds() - FileScanStartTime);
