CookOnTheFly (platform protoco): Added an additional check to ITargetDeviceSocket::IsConnected in the cooking server. It's needed because on some platforms ITargetDevice::IsProtocolAvailable doesn't have the necessary knowledge and unconditionally returns 'true' (we only know if the protocol is available after creating a socket).

#rnx

diff --git a/Engine/Source/Runtime/NetworkFileSystem/Private/NetworkFileServerPlatformProtocol.cpp b/Engine/Source/Runtime/NetworkFileSystem/Private/NetworkFileServerPlatformProtocol.cpp
--- a/Engine/Source/Runtime/NetworkFileSystem/Private/NetworkFileServerPlatformProtocol.cpp
+++ b/Engine/Source/Runtime/NetworkFileSystem/Private/NetworkFileServerPlatformProtocol.cpp
@@ -330,11 +330,18 @@ void FNetworkFileServerPlatformProtocol::AddConnectionsForNewDevices(ITargetPlat
 
 					if (Socket)
 					{
-						FConnectionThreaded* Connection =
-							new FConnectionThreaded(Device, Socket, FileServerOptions);
-
-						Connections.Add(Connection);
-						UE_LOG(LogFileServer, Display, TEXT("Client %s connected."), *Connection->GetName());
+						if (Socket->Connected())
+						{
+							FConnectionThreaded* Connection =
+								new FConnectionThreaded(Device, Socket, FileServerOptions);
+
+							Connections.Add(Connection);
+							UE_LOG(LogFileServer, Display, TEXT("Client %s connected."), *Connection->GetName());
+						}
+						else
+						{
+							Device->CloseConnection(Socket);
+						}
 					}
 				}
 			}
