CookOnTheFly NetworkPlatformFile: Fix typo preventing server messages from being processed on the client.
Fix editor crash on shutdown.
#rnx

diff --git a/Engine/Source/Runtime/NetworkFile/Private/NetworkPlatformFile.cpp b/Engine/Source/Runtime/NetworkFile/Private/NetworkPlatformFile.cpp
--- a/Engine/Source/Runtime/NetworkFile/Private/NetworkPlatformFile.cpp
+++ b/Engine/Source/Runtime/NetworkFile/Private/NetworkPlatformFile.cpp
@@ -183,7 +183,7 @@ bool FNetworkPlatformFile::ReceiveResponse(TArray<uint8> &Out )
 	{
 		NewPayloadEvent->Wait();
 		TOptional<TArray<uint8>> Payload = PendingPayloads.Dequeue();
-		if (!Payload.IsSet())
+		if (Payload.IsSet())
 		{
 			Out = MoveTemp(Payload.GetValue());
 			break;
diff --git a/Engine/Source/Runtime/NetworkFileSystem/Private/NetworkFileServer.cpp b/Engine/Source/Runtime/NetworkFileSystem/Private/NetworkFileServer.cpp
--- a/Engine/Source/Runtime/NetworkFileSystem/Private/NetworkFileServer.cpp
+++ b/Engine/Source/Runtime/NetworkFileSystem/Private/NetworkFileServer.cpp
@@ -138,6 +138,10 @@ int32 FNetworkFileServer::NumConnections(void) const
 
 void FNetworkFileServer::Shutdown(void)
 {
+	if (!CookOnTheFlyServer)
+	{
+		return;
+	}
 	CookOnTheFlyServer->OnClientConnected().RemoveAll(this);
 	CookOnTheFlyServer->OnClientDisconnected().RemoveAll(this);
 	CookOnTheFlyServer->OnRequest(UE::Cook::ECookOnTheFlyMessage::NetworkPlatformFile).Unbind();
