Fixed some potential integer overflows in Zen/COTF communication layers

diff --git a/Engine/Source/Developer/Zen/Internal/ZenServerHttp.cpp b/Engine/Source/Developer/Zen/Internal/ZenServerHttp.cpp
--- a/Engine/Source/Developer/Zen/Internal/ZenServerHttp.cpp
+++ b/Engine/Source/Developer/Zen/Internal/ZenServerHttp.cpp
@@ -577,7 +577,7 @@ namespace UE::Zen {
 	size_t FZenHttpRequest::FStatics::StaticWriteBodyFn(void* Ptr, size_t SizeInBlocks, size_t BlockSizeInBytes, void* UserData)
 	{
 		FZenHttpRequest* Request = static_cast<FZenHttpRequest*>(UserData);
-		const size_t WriteSize = SizeInBlocks * BlockSizeInBytes;
+		const int64 WriteSize = IntCastChecked<int64>(SizeInBlocks * BlockSizeInBytes);
 		TArray64<uint8>* WriteDataBufferPtr = Request->WriteDataBufferPtr;
 
 		if (WriteDataBufferPtr && WriteSize > 0)
@@ -591,7 +591,7 @@ namespace UE::Zen {
 
 				if (const ANSICHAR* ContentLengthHeader = FCStringAnsi::Strstr(Header, ContentLengthHeaderStr))
 				{
-					size_t ContentLength = (size_t)FCStringAnsi::Atoi64(ContentLengthHeader + strlen(ContentLengthHeaderStr));
+					int64 ContentLength = FCStringAnsi::Atoi64(ContentLengthHeader + strlen(ContentLengthHeaderStr));
 					if (ContentLength > 0)
 					{
 						WriteDataBufferPtr->Reserve(ContentLength);
diff --git a/Engine/Source/Runtime/NetworkFileSystem/Private/NetworkFileServerConnection.cpp b/Engine/Source/Runtime/NetworkFileSystem/Private/NetworkFileServerConnection.cpp
--- a/Engine/Source/Runtime/NetworkFileSystem/Private/NetworkFileServerConnection.cpp
+++ b/Engine/Source/Runtime/NetworkFileSystem/Private/NetworkFileServerConnection.cpp
@@ -505,28 +505,41 @@ void FNetworkFileServerClientConnection::ProcessReadFile( FArchive& In, FArchive
 	int64 BytesToRead = 0;
 	In << BytesToRead;
 
-	int64 BytesRead = 0;
 	IFileHandle* File = FindOpenFile(HandleId);
 
 	if (File)
 	{
-		uint8* Dest = (uint8*)FMemory::Malloc(BytesToRead);		
-
-		if (File->Read(Dest, BytesToRead))
+		constexpr int64 BufferSize = 4 << 20;
+		uint8* Buffer = (uint8*)FMemory::Malloc(BufferSize);
+		bool bIsFirstRead = true;
+		while (BytesToRead > 0)
 		{
-			BytesRead = BytesToRead;
-			Out << BytesRead;
-			Out.Serialize(Dest, BytesRead);
-		}
-		else
-		{
-			Out << BytesRead;
+			int64 CappedBytesToRead = FMath::Min(BufferSize, BytesToRead);
+			if (!File->Read(Buffer, CappedBytesToRead))
+			{
+				if (bIsFirstRead)
+				{
+					int64 BytesRead = 0;
+					Out << BytesRead;
+					break;
+				}
+				// If this is not the first read we've already written the expected number of bytes to the stream so we have to deliver on that
+				FMemory::Memset(Buffer, 0, CappedBytesToRead);
+			}
+			else if (bIsFirstRead)
+			{
+				Out << BytesToRead;
+			}
+			Out.Serialize(Buffer, CappedBytesToRead);
+			BytesToRead -= CappedBytesToRead;
+			bIsFirstRead = false;
+			
 		}
-
-		FMemory::Free(Dest);
+		FMemory::Free(Buffer);
 	}
 	else
 	{
+		int64 BytesRead = 0;
 		Out << BytesRead;
 	}
 }
@@ -546,15 +559,20 @@ void FNetworkFileServerClientConnection::ProcessWriteFile( FArchive& In, FArchiv
 		int64 BytesToWrite = 0;
 		In << BytesToWrite;
 
-		uint8* Source = (uint8*)FMemory::Malloc(BytesToWrite);
-		In.Serialize(Source, BytesToWrite);
-
-		if (File->Write(Source, BytesToWrite))
+		constexpr int64 BufferSize = 4 << 20;
+		uint8* Buffer = (uint8*)FMemory::Malloc(BufferSize);
+		while (BytesToWrite > 0)
 		{
-			BytesWritten = BytesToWrite;
+			int64 CappedBytesToWrite = FMath::Min(BufferSize, BytesToWrite);
+			In.Serialize(Buffer, CappedBytesToWrite);
+			if (!File->Write(Buffer, CappedBytesToWrite))
+			{
+				break;
+			}
+			BytesWritten += CappedBytesToWrite;
+			BytesToWrite -= CappedBytesToWrite;
 		}
-
-		FMemory::Free(Source); 
+		FMemory::Free(Buffer);
 	}
 		
 	Out << BytesWritten;
diff --git a/Engine/Source/Runtime/StorageServerClient/Private/StorageServerConnection.cpp b/Engine/Source/Runtime/StorageServerClient/Private/StorageServerConnection.cpp
--- a/Engine/Source/Runtime/StorageServerClient/Private/StorageServerConnection.cpp
+++ b/Engine/Source/Runtime/StorageServerClient/Private/StorageServerConnection.cpp
@@ -179,7 +179,7 @@ FStorageServerResponse::FStorageServerResponse(FStorageServerConnection& InOwner
 		ResponseLine = ReadResponseLine();
 		if (ResponseLine.StartsWith("Content-Length: "))
 		{
-			ContentLength = TCString<ANSICHAR>::Atoi64(ResponseLine.GetData() + 16);
+			ContentLength = FMath::Max(0, TCString<ANSICHAR>::Atoi64(ResponseLine.GetData() + 16));
 		}
 		else if (ResponseLine.StartsWith("Content-Type: "))
 		{
@@ -258,6 +258,11 @@ bool FStorageServerChunkBatchRequest::Issue(TFunctionRef<void(uint32 ChunkCount,
 		return false;
 	}
 	Response << ChunkCount;
+	if (ChunkCount > INT32_MAX)
+	{
+		UE_LOG(LogStorageServerConnection, Fatal, TEXT("Invalid chunk count in chunk batch response from storage server."));
+		return false;
+	}
 	Response << Reserved1;
 	Response << Reserved2;
 
@@ -341,7 +346,7 @@ int64 FStorageServerResponse::SerializeChunk(FStorageServerSerializationContext&
 
 	if (ContentType == EStorageServerContentType::Binary)
 	{
-		const uint64 ChunkSize = FMath::Min(ContentLength, RawSize);
+		const uint64 ChunkSize = FMath::Min<uint64>(ContentLength, RawSize);
 		OutChunk = TargetVa ? FIoBuffer(FIoBuffer::Wrap, TargetVa, ChunkSize) : FIoBuffer(ChunkSize);
 		Serialize(OutChunk.Data(), OutChunk.DataSize());
 		return ChunkSize;
@@ -387,7 +392,7 @@ int64 FStorageServerResponse::SerializeChunkTo(FMutableMemoryView Memory, uint64
 
 	if (ContentType == EStorageServerContentType::Binary)
 	{
-		FMutableMemoryView Dst = Memory.Left(FMath::Min(Memory.GetSize(), ContentLength));
+		FMutableMemoryView Dst = Memory.Left(FMath::Min<uint64>(Memory.GetSize(), ContentLength));
 		Serialize(Dst.GetData(), Dst.GetSize());
 		return Dst.GetSize();
 	}
diff --git a/Engine/Source/Runtime/StorageServerClient/Public/StorageServerConnection.h b/Engine/Source/Runtime/StorageServerClient/Public/StorageServerConnection.h
--- a/Engine/Source/Runtime/StorageServerClient/Public/StorageServerConnection.h
+++ b/Engine/Source/Runtime/StorageServerClient/Public/StorageServerConnection.h
@@ -149,8 +149,8 @@ private:
 
 	FStorageServerConnection& Owner;
 	FSocket* Socket = nullptr;
-	uint64 ContentLength = 0;
-	uint64 Position = 0;
+	int64 ContentLength = 0;
+	int64 Position = 0;
 	int32 ErrorCode;
 	FString ErrorMessage;
 	bool bIsOk = false;
