UE-176366 CookOnTheFly does not apply PluginExtensions.
-First problem: we were filtering out plugins based on the FPlugin, rather than the FPluginReference.  This meant that the platform extension platforms were not in the supported list.  Needed to use the FPluginReference in that one spot.
-Second problem: platform extension .uplugin files were not added to the file list for the NetworkFileServer because we find Plugins separately specifically looking for the basedir of each plugin.  In order to find these files I am caching the basedir of each child plugin in the FPlugin when we merge the children into it.  NetworkFileServer can then find these directores, examine them to see if they are not for the targeted platform, and make them available to the client.

diff --git a/Engine/Source/Runtime/NetworkFileSystem/Private/NetworkFileServerConnection.cpp b/Engine/Source/Runtime/NetworkFileSystem/Private/NetworkFileServerConnection.cpp
--- a/Engine/Source/Runtime/NetworkFileSystem/Private/NetworkFileServerConnection.cpp
+++ b/Engine/Source/Runtime/NetworkFileSystem/Private/NetworkFileServerConnection.cpp
@@ -3,6 +3,7 @@
 #include "NetworkFileServerConnection.h"
 #include "HAL/PlatformFileManager.h"
 #include "Misc/Paths.h"
+#include "Misc/PathViews.h"
 #include "Misc/ScopeLock.h"
 #include "Serialization/BufferArchive.h"
 #include "Misc/ConfigCacheIni.h"
@@ -1076,14 +1077,7 @@ bool FNetworkFileServerClientConnection::ProcessGetFileList( FArchive& In, FArch
 		Sandbox->IterateDirectory(*RootDirectories[DirIndex], Visitor);
 	}
 
-	// Traverse plugin directories
-	TArray<TSharedRef<IPlugin>> AllPlugins = IPluginManager::Get().GetDiscoveredPlugins();
-	for (TSharedRef<IPlugin> Plugin : AllPlugins)
-	{
-		ScanExtensionRootDirectory(Sandbox.Get(), Plugin->GetBaseDir(), RootDirectories, Visitor.FileTimes);
-	}
-
-	// Traverse platform extension directories
+	// Get PlatformDirectoryNames
 	FString ServerEnginePlatformExtensionsRelativePath = EnginePlatformExtensionsRelativePath;
 	ConvertClientFilenameToServerFilename(ServerEnginePlatformExtensionsRelativePath);
 	FString ServerProjectPlatformExtensionsRelativePath = ProjectPlatformExtensionsRelativePath;
@@ -1105,6 +1099,56 @@ bool FNetworkFileServerClientConnection::ProcessGetFileList( FArchive& In, FArch
 			PlatformDirectoryNames.AddUnique(PlatformName);
 		}
 	}
+
+	// Traverse plugin directories
+	TSet<FString> PlatformDirectoryNameSet;
+	PlatformDirectoryNameSet.Append(PlatformDirectoryNames);
+	TArray<TSharedRef<IPlugin>> AllPlugins = IPluginManager::Get().GetDiscoveredPlugins();
+	for (TSharedRef<IPlugin> Plugin : AllPlugins)
+	{
+		// First the base directory of the plugin.
+		ScanExtensionRootDirectory(Sandbox.Get(), Plugin->GetBaseDir(), RootDirectories, Visitor.FileTimes);
+
+		// Next the plugin extension directories of this plugin.
+		TArray<FString> PluginExtensionDirs = Plugin->GetExtensionBaseDirs();
+		for (const FString& ExtensionDir : PluginExtensionDirs)
+		{
+			// Scan for Platforms/X.  If X is not one of our platforms do not scan this extension directory.  
+			// If X is one of our platforms or this extension is not Platforms restricted at all scan it.
+			bool bFoundPlatforms = false;
+			bool bDone = false;
+			bool bWrongPlatform = false;
+			FPathViews::IterateComponents(
+				ExtensionDir,
+				[&bFoundPlatforms, &bDone, &bWrongPlatform, &PlatformDirectoryNameSet](FStringView CurrentPathComponent)
+				{
+					if (!bFoundPlatforms)
+					{
+						if (CurrentPathComponent == FString(TEXT("Platforms")))
+						{
+							bFoundPlatforms = true;
+						}
+					}
+					else if (!bDone)
+					{
+						bWrongPlatform = !PlatformDirectoryNameSet.Contains(FString(CurrentPathComponent));
+						bDone = true;
+					}
+					else
+					{
+						// Do nothing.
+					}
+				}
+			);
+
+			if (!bWrongPlatform)
+			{
+				ScanExtensionRootDirectory(Sandbox.Get(), ExtensionDir, RootDirectories, Visitor.FileTimes);
+			}
+		}
+	}
+
+	// Traverse platform extension directories
 	for (const FString& PlatformDirectoryName : PlatformDirectoryNames)
 	{
 		ScanExtensionRootDirectory(Sandbox.Get(), ServerEnginePlatformExtensionsRelativePath / PlatformDirectoryName, RootDirectories, Visitor.FileTimes);
diff --git a/Engine/Source/Runtime/Projects/Private/PluginManager.cpp b/Engine/Source/Runtime/Projects/Private/PluginManager.cpp
--- a/Engine/Source/Runtime/Projects/Private/PluginManager.cpp
+++ b/Engine/Source/Runtime/Projects/Private/PluginManager.cpp
@@ -376,6 +376,17 @@ FString FPlugin::GetBaseDir() const
 	return FPaths::GetPath(FileName);
 }
 
+TArray<FString> FPlugin::GetExtensionBaseDirs() const
+{
+	TArray<FString> OutDirs;
+	OutDirs.Reserve(PluginExtensionFileNameList.Num());
+	for (const FString& ExtensionFileName : PluginExtensionFileNameList)
+	{
+		OutDirs.Emplace(FPaths::GetPath(ExtensionFileName));
+	}
+	return OutDirs;
+}
+
 FString FPlugin::GetContentDir() const
 {
 	return FPaths::GetPath(FileName) / TEXT("Content");
@@ -895,6 +906,9 @@ void FPluginManager::ReadAllPlugins(FDiscoveredPluginMap& Plugins, const TSet<FS
 					Parent->Descriptor.SupportedTargetPlatforms.AddUnique(SupportedTargetPlatform);
 				}
 			}
+
+			// we need to remember the child's filename so we can make the file findable on the network file system
+			Parent->PluginExtensionFileNameList.Add(Child->FileName);
 		}
 		else
 		{
@@ -2033,7 +2047,7 @@ bool FPluginManager::ConfigureEnabledPluginForTarget(const FPluginReferenceDescr
 			}
 
 			// Check the plugin supports this platform
-			if(!bLoadPluginsForTargetPlatforms && !Plugin.Descriptor.SupportsTargetPlatform(Platform))
+			if(!bLoadPluginsForTargetPlatforms && !Reference.IsSupportedTargetPlatform(Platform))
 			{
 				UE_LOG(LogPluginManager, Verbose, TEXT("Ignoring plugin '%s' due to unsupported platform in plugin descriptor"), *Reference.Name);
 				continue;
diff --git a/Engine/Source/Runtime/Projects/Private/PluginManager.h b/Engine/Source/Runtime/Projects/Private/PluginManager.h
--- a/Engine/Source/Runtime/Projects/Private/PluginManager.h
+++ b/Engine/Source/Runtime/Projects/Private/PluginManager.h
@@ -21,6 +21,9 @@ public:
 	/** The filename that the plugin was loaded from */
 	FString FileName;
 
+	/** The filenames from which any child plugins that have been merged were loaded */
+	TArray<FString> PluginExtensionFileNameList;
+
 	/** The plugin's settings */
 	FPluginDescriptor Descriptor;
 
@@ -66,6 +69,7 @@ public:
 	}
 
 	virtual FString GetBaseDir() const override;
+	virtual TArray<FString> GetExtensionBaseDirs() const override;
 	virtual FString GetContentDir() const override;
 	virtual FString GetMountedAssetPath() const override;
 
diff --git a/Engine/Source/Runtime/Projects/Public/Interfaces/IPluginManager.h b/Engine/Source/Runtime/Projects/Public/Interfaces/IPluginManager.h
--- a/Engine/Source/Runtime/Projects/Public/Interfaces/IPluginManager.h
+++ b/Engine/Source/Runtime/Projects/Public/Interfaces/IPluginManager.h
@@ -101,6 +101,13 @@ public:
 	 */
 	virtual FString GetBaseDir() const = 0;
 
+	/**
+	 * Get a filesystem path to the plugin's directory.
+	 *
+	 * @return Filesystem path to the plugin's base directory.
+	 */
+	virtual TArray<FString> GetExtensionBaseDirs() const = 0;
+
 	/**
 	 * Get a filesystem path to the plugin's content directory.
 	 *
