CookOnTheFlyServer: Suppress the scanning of the Engine\Restricted folder when reporting the available file list to the connected client, because the Restricted folder can have a large number of files. It should not be required for the client to access any files in the Restricted folder, except for in-use plugins which are scanned individually.
#rnx

diff --git a/Engine/Source/Runtime/NetworkFileSystem/Private/NetworkFileServerConnection.cpp b/Engine/Source/Runtime/NetworkFileSystem/Private/NetworkFileServerConnection.cpp
--- a/Engine/Source/Runtime/NetworkFileSystem/Private/NetworkFileServerConnection.cpp
+++ b/Engine/Source/Runtime/NetworkFileSystem/Private/NetworkFileServerConnection.cpp
@@ -839,6 +839,7 @@ static void AddDirectoriesToIgnore(const FString& RootDir, TArray<FString>& OutD
 	OutDirectoriesToSkip.Add(FString(RootDir / TEXT("Programs")));
 	OutDirectoriesToSkip.Add(FString(RootDir / TEXT("Platforms")));
 	OutDirectoriesToSkip.Add(FString(RootDir / TEXT("Build")));
+	OutDirectoriesToSkip.Add(FString(RootDir / TEXT("Restricted")));
 	OutDirectoriesToNotRecurse.Add(FString(RootDir / TEXT("DerivedDataCache")));
 }
 
