Removed defunct HTTP server for serving files to HTML5 platform

#rnx

diff --git a/Engine/Source/Runtime/NetworkFileSystem/NetworkFileSystem.Build.cs b/Engine/Source/Runtime/NetworkFileSystem/NetworkFileSystem.Build.cs
--- a/Engine/Source/Runtime/NetworkFileSystem/NetworkFileSystem.Build.cs
+++ b/Engine/Source/Runtime/NetworkFileSystem/NetworkFileSystem.Build.cs
@@ -25,18 +25,6 @@ namespace UnrealBuildTool.Rules
 					"Sockets",
 				});
 
-			if ((Target.Platform == UnrealTargetPlatform.Win64) ||
-				(Target.Platform == UnrealTargetPlatform.Mac))
-			{
-				AddEngineThirdPartyPrivateStaticDependencies(Target, "OpenSSL", "libWebSockets", "zlib");
-				PublicDefinitions.Add("ENABLE_HTTP_FOR_NFS=1");
-				PrivateDependencyModuleNames.Add("SSL");
-			}
-			else
-			{
-				PublicDefinitions.Add("ENABLE_HTTP_FOR_NFS=0");
-			}
-
 			PrecompileForTargets = PrecompileTargetsType.Editor;
 
 			UnsafeTypeCastWarningLevel = WarningLevel.Error;
diff --git a/Engine/Source/Runtime/NetworkFileSystem/Private/NetworkFileServerHttp.cpp b/Engine/Source/Runtime/NetworkFileSystem/Private/NetworkFileServerHttp.cpp
deleted file mode 100644
diff --git a/Engine/Source/Runtime/NetworkFileSystem/Private/NetworkFileServerHttp.h b/Engine/Source/Runtime/NetworkFileSystem/Private/NetworkFileServerHttp.h
deleted file mode 100644
diff --git a/Engine/Source/Runtime/NetworkFileSystem/Private/NetworkFileSystemModule.cpp b/Engine/Source/Runtime/NetworkFileSystem/Private/NetworkFileSystemModule.cpp
--- a/Engine/Source/Runtime/NetworkFileSystem/Private/NetworkFileSystemModule.cpp
+++ b/Engine/Source/Runtime/NetworkFileSystem/Private/NetworkFileSystemModule.cpp
@@ -6,7 +6,6 @@
 #include "INetworkFileSystemModule.h"
 #include "NetworkFileSystemLog.h"
 #include "NetworkFileServer.h"
-#include "NetworkFileServerHttp.h"
 #include "Interfaces/ITargetPlatformManagerModule.h"
 #include "CookOnTheFlyNetServer.h"
 #include "IPAddress.h"
@@ -55,10 +54,6 @@ public:
 
 		switch (FileServerOptions.Protocol)
 		{
-#if ENABLE_HTTP_FOR_NFS
-		case NFSP_Http: 
-			return new FNetworkFileServerHttp(MoveTemp(FileServerOptions));
-#endif
 		case NFSP_Tcp:
 		case NFSP_Platform:
 			UE::Cook::FCookOnTheFlyNetworkServerOptions CookOnTheFlyServerOptions;
